Sequencer: Control Rig: New Motion Trail and Pivot tools

diff --git a/Engine/Plugins/Animation/ControlRig/Source/ControlRig/Private/Sequencer/MovieSceneControlRigParameterTrack.cpp b/Engine/Plugins/Animation/ControlRig/Source/ControlRig/Private/Sequencer/MovieSceneControlRigParameterTrack.cpp
--- a/Engine/Plugins/Animation/ControlRig/Source/ControlRig/Private/Sequencer/MovieSceneControlRigParameterTrack.cpp
+++ b/Engine/Plugins/Animation/ControlRig/Source/ControlRig/Private/Sequencer/MovieSceneControlRigParameterTrack.cpp
@@ -318,7 +318,15 @@ void UMovieSceneControlRigParameterTrack::SetSectionToKey(UMovieSceneSection* In
 
 UMovieSceneSection* UMovieSceneControlRigParameterTrack::GetSectionToKey() const
 {
-	return SectionToKey;
+	if (SectionToKey)
+	{
+		return SectionToKey;
+	}
+	else if(Sections.Num() >0)
+	{
+		return Sections[0];
+	}
+	return nullptr;
 }
 
 void UMovieSceneControlRigParameterTrack::ReconstructControlRig()
diff --git a/Engine/Plugins/Animation/ControlRig/Source/ControlRig/Public/Rigs/RigHierarchyDefines.h b/Engine/Plugins/Animation/ControlRig/Source/ControlRig/Public/Rigs/RigHierarchyDefines.h
--- a/Engine/Plugins/Animation/ControlRig/Source/ControlRig/Public/Rigs/RigHierarchyDefines.h
+++ b/Engine/Plugins/Animation/ControlRig/Source/ControlRig/Public/Rigs/RigHierarchyDefines.h
@@ -842,6 +842,30 @@ FORCEINLINE_DEBUGGABLE FQuat FRigControlValue::SetFromString<FQuat>(const FStrin
 	return Value;
 }
 
+enum class EControlRigContextChannelToKey : uint32
+{
+	None = 0x000,
+
+	TranslationX = 0x001,
+	TranslationY = 0x002,
+	TranslationZ = 0x004,
+	Translation = TranslationX | TranslationY | TranslationZ,
+
+	RotationX = 0x008,
+	RotationY = 0x010,
+	RotationZ = 0x020,
+	Rotation = RotationX | RotationY | RotationZ,
+
+	ScaleX = 0x040,
+	ScaleY = 0x080,
+	ScaleZ = 0x100,
+	Scale = ScaleX | ScaleY | ScaleZ,
+
+	AllTransform = Translation | Rotation | Scale,
+
+};
+ENUM_CLASS_FLAGS(EControlRigContextChannelToKey)
+
 USTRUCT(BlueprintType)
 struct CONTROLRIG_API FRigControlModifiedContext
 {
@@ -849,25 +873,30 @@ struct CONTROLRIG_API FRigControlModifiedContext
 
 	FRigControlModifiedContext()
 	: SetKey(EControlRigSetKey::DoNotCare)
+	, KeyMask((uint32)EControlRigContextChannelToKey::AllTransform)
 	, LocalTime(FLT_MAX)
 	, EventName(NAME_None)
 	{}
 
 	FRigControlModifiedContext(EControlRigSetKey InSetKey)
 		: SetKey(InSetKey)
+		, KeyMask((uint32)EControlRigContextChannelToKey::AllTransform)
 		, LocalTime(FLT_MAX)
 		, EventName(NAME_None)
 	{}
 
-	FRigControlModifiedContext(EControlRigSetKey InSetKey, float InLocalTime, const FName& InEventName = NAME_None)
+	FRigControlModifiedContext(EControlRigSetKey InSetKey, float InLocalTime, const FName& InEventName = NAME_None, EControlRigContextChannelToKey InKeyMask = EControlRigContextChannelToKey::AllTransform)
 		: SetKey(InSetKey)
+		, KeyMask((uint32)InKeyMask)
 		, LocalTime(InLocalTime)
 		, EventName(InEventName)
 	{}
-
+	
 	EControlRigSetKey SetKey;
+	uint32 KeyMask;
 	float LocalTime;
 	FName EventName;
+	
 };
 
 /*
diff --git a/Engine/Plugins/Animation/ControlRig/Source/ControlRigEditor/ControlRigEditor.Build.cs b/Engine/Plugins/Animation/ControlRig/Source/ControlRigEditor/ControlRigEditor.Build.cs
--- a/Engine/Plugins/Animation/ControlRig/Source/ControlRigEditor/ControlRigEditor.Build.cs
+++ b/Engine/Plugins/Animation/ControlRig/Source/ControlRigEditor/ControlRigEditor.Build.cs
@@ -62,6 +62,7 @@ namespace UnrealBuildTool.Rules
 					"ContentBrowserData",
                     "LevelEditor",
                     "SceneOutliner",
+                    "EditorInteractiveToolsFramework",
                     "LevelSequence",
                     "GraphEditor",
                     "PropertyPath",
diff --git a/Engine/Plugins/Animation/ControlRig/Source/ControlRigEditor/Private/ControlRigEditorModule.cpp b/Engine/Plugins/Animation/ControlRig/Source/ControlRigEditor/Private/ControlRigEditorModule.cpp
--- a/Engine/Plugins/Animation/ControlRig/Source/ControlRigEditor/Private/ControlRigEditorModule.cpp
+++ b/Engine/Plugins/Animation/ControlRig/Source/ControlRigEditor/Private/ControlRigEditorModule.cpp
@@ -115,6 +115,7 @@
 #include "EditMode/SControlRigTweenWidget.h"
 #include "EditMode/SControlRigSnapper.h"
 #include "Dialogs/CustomDialog.h"
+#include "Tools/SMotionTrailOptions.h"
 
 #define LOCTEXT_NAMESPACE "ControlRigEditorModule"
 
@@ -123,7 +124,7 @@ DEFINE_LOG_CATEGORY(LogControlRigEditor);
 const FName IControlRigEditorModule::ControlRigPoseTab = FName("ControlRigPoseTab");
 const FName IControlRigEditorModule::ControlRigTweenTab = FName("ControlRigTweenTab");
 const FName IControlRigEditorModule::ControlRigSnapperTab = FName("ControlRigSnapperTab");
-const FName IControlRigEditorModule::ControlRigTempPivotTab = FName("ControlRigTempPivotTab");
+const FName IControlRigEditorModule::ControlRigMotionTrailTab = FName("ControlRigMotionTrailTab");
 
 
 TSharedRef<SDockTab> SpawnRigProfiler( const FSpawnTabArgs& Args )
@@ -162,15 +163,16 @@ TSharedRef<SDockTab> SpawnSnapperTab(const FSpawnTabArgs& Args)
 		];
 }
 
-TSharedRef<SDockTab> SpawnPivotTab(const FSpawnTabArgs& Args)
+TSharedRef<SDockTab> SpawnMotionTrailTab(const FSpawnTabArgs& Args)
 {
 	return SNew(SDockTab)
 		.TabRole(ETabRole::NomadTab)
 		[
-			SNew(SControlRigTweenWidget)
+			SNew(SMotionTrailOptions)
 		];
 }
 
+
 void FControlRigEditorModule::StartupModule()
 {
 	FControlRigEditModeCommands::Register();
@@ -303,13 +305,12 @@ void FControlRigEditorModule::StartupModule()
 			.SetIcon(FSlateIcon(FEditorStyle::GetStyleSetName(), "Persona.Tabs.AnimSlotManager"));//MZ todo replace with correct icon SkeletonTree.SkeletonSocket
 		FGlobalTabmanager::Get()->RegisterDefaultTabWindowSize(IControlRigEditorModule::ControlRigSnapperTab, FVector2D(400, 400));
 
-		FGlobalTabmanager::Get()->RegisterNomadTabSpawner(IControlRigEditorModule::ControlRigTempPivotTab, FOnSpawnTab::CreateStatic(&SpawnPivotTab))
-			.SetDisplayName(NSLOCTEXT("UnrealEditor", "ControlRigTempPivotTab", "Control Rig Temp Pivot Tab"))
-			.SetTooltipText(NSLOCTEXT("UnrealEditor", "ControlRigTweenTooltip", "Open the Control Rig Pivot tab."))
+		FGlobalTabmanager::Get()->RegisterNomadTabSpawner(IControlRigEditorModule::ControlRigMotionTrailTab, FOnSpawnTab::CreateStatic(&SpawnMotionTrailTab))
+			.SetDisplayName(NSLOCTEXT("UnrealEditor", "MotionTrailTab", "Motion Trail"))
+			.SetTooltipText(NSLOCTEXT("UnrealEditor", "MotionTrailTabTooltip", "Open the Motion Trail tab."))
 			.SetMenuType(ETabSpawnerMenuType::Hidden)
-			.SetIcon(FSlateIcon(FEditorStyle::GetStyleSetName(), "SkeletonTree.BlendProfile")); //MZ todo replace with correct icon SkeletonTree.LODBone
-		FGlobalTabmanager::Get()->RegisterDefaultTabWindowSize(IControlRigEditorModule::ControlRigTempPivotTab, FVector2D(400, 100));
-
+			.SetIcon(FSlateIcon(FEditorStyle::GetStyleSetName(), "Persona.Tabs.AnimSlotManager"));//MZ todo replace with correct icon SkeletonTree.SkeletonSocket
+		FGlobalTabmanager::Get()->RegisterDefaultTabWindowSize(IControlRigEditorModule::ControlRigMotionTrailTab, FVector2D(200, 500));
 
 	};
 #endif
@@ -338,7 +339,7 @@ void FControlRigEditorModule::ShutdownModule()
 		FGlobalTabmanager::Get()->UnregisterNomadTabSpawner(IControlRigEditorModule::ControlRigPoseTab);
 		FGlobalTabmanager::Get()->UnregisterNomadTabSpawner(IControlRigEditorModule::ControlRigTweenTab);
 		FGlobalTabmanager::Get()->UnregisterNomadTabSpawner(IControlRigEditorModule::ControlRigSnapperTab);
-		FGlobalTabmanager::Get()->UnregisterNomadTabSpawner(IControlRigEditorModule::ControlRigTempPivotTab);
+		FGlobalTabmanager::Get()->UnregisterNomadTabSpawner(IControlRigEditorModule::ControlRigMotionTrailTab);
 	}
 #endif
 
diff --git a/Engine/Plugins/Animation/ControlRig/Source/ControlRigEditor/Private/EditMode/ControlRigEditMode.cpp b/Engine/Plugins/Animation/ControlRig/Source/ControlRigEditor/Private/EditMode/ControlRigEditMode.cpp
--- a/Engine/Plugins/Animation/ControlRig/Source/ControlRigEditor/Private/EditMode/ControlRigEditMode.cpp
+++ b/Engine/Plugins/Animation/ControlRig/Source/ControlRigEditor/Private/EditMode/ControlRigEditMode.cpp
@@ -861,6 +861,16 @@ bool FControlRigEditMode::HandleClick(FEditorViewportClient* InViewportClient, H
 					}
 					else
 					{
+						//also need to clear actor selection. Sequencer will handle this automatically if done in Sequencder UI but not if done by clicking
+						if (IsInLevelEditor())
+						{
+							if (GEditor && GEditor->GetSelectedActorCount())
+							{
+								const FScopedTransaction Transaction(NSLOCTEXT("Sequencer", "UpdatingActorComponentSelectionNone", "Select None"), !GIsTransacting);
+								GEditor->SelectNone(false, true);
+								GEditor->NoteSelectionChange();
+							}
+						}
 						ClearRigElementSelection(FRigElementTypeHelper::ToMask(ERigElementType::Control));
 						SetRigElementSelection(ERigElementType::Control, ControlName, true);
 					}
diff --git a/Engine/Plugins/Animation/ControlRig/Source/ControlRigEditor/Private/EditMode/SControlRigEditModeTools.cpp b/Engine/Plugins/Animation/ControlRig/Source/ControlRigEditor/Private/EditMode/SControlRigEditModeTools.cpp
--- a/Engine/Plugins/Animation/ControlRig/Source/ControlRigEditor/Private/EditMode/SControlRigEditModeTools.cpp
+++ b/Engine/Plugins/Animation/ControlRig/Source/ControlRigEditor/Private/EditMode/SControlRigEditModeTools.cpp
@@ -29,6 +29,10 @@
 #include "IControlRigEditorModule.h"
 #include "Framework/Docking/TabManager.h"
 #include "ControlRigEditorStyle.h"
+#include "LevelEditor.h"
+#include "EditorModeManager.h"
+#include "InteractiveToolManager.h"
+#include "EdModeInteractiveToolsContext.h"
 
 #define LOCTEXT_NAMESPACE "ControlRigRootCustomization"
 
@@ -515,21 +519,22 @@ void SControlRigEditModeTools::CustomizeToolBarPalette(FToolBarBuilder& ToolBarB
 	ToolBarBuilder.AddToolBarButton(
 		FUIAction(
 			FExecuteAction::CreateLambda([this] {
-		FControlRigEditMode* ControlRigEditMode = static_cast<FControlRigEditMode*>(ModeTools->GetActiveMode(FControlRigEditMode::ModeName));
-		if (ControlRigEditMode)
-		{
-			ControlRigEditMode->SetOnlySelectRigControls(!ControlRigEditMode->GetOnlySelectRigControls());
-		}
-	}),
+			FControlRigEditMode* ControlRigEditMode = static_cast<FControlRigEditMode*>(ModeTools->GetActiveMode(FControlRigEditMode::ModeName));
+			if (ControlRigEditMode)
+			{
+				ControlRigEditMode->SetOnlySelectRigControls(!ControlRigEditMode->GetOnlySelectRigControls());
+			}
+			}),
 			FCanExecuteAction(),
-		FIsActionChecked::CreateLambda([this] {
-		FControlRigEditMode* ControlRigEditMode = static_cast<FControlRigEditMode*>(ModeTools->GetActiveMode(FControlRigEditMode::ModeName));
-		if (ControlRigEditMode)
-		{
-			return ControlRigEditMode->GetOnlySelectRigControls();
-		}
-		return false;
-	})
+			FIsActionChecked::CreateLambda([this] {
+			FControlRigEditMode* ControlRigEditMode = static_cast<FControlRigEditMode*>(ModeTools->GetActiveMode(FControlRigEditMode::ModeName));
+			if (ControlRigEditMode)
+			{
+				return ControlRigEditMode->GetOnlySelectRigControls();
+			}
+			return false;
+			})
+
 		),
 		NAME_None,
 		LOCTEXT("OnlySelectControls", "Select"),
@@ -569,17 +574,46 @@ void SControlRigEditModeTools::CustomizeToolBarPalette(FToolBarBuilder& ToolBarB
 		FSlateIcon(TEXT("ControlRigEditorStyle"), TEXT("ControlRig.SnapperTool")),
 		EUserInterfaceActionType::Button
 	);
-	/*
-	// Pivot
+
+	// Motion Trail
 	ToolBarBuilder.AddToolBarButton(
-		FExecuteAction::CreateSP(this, &SControlRigEditModeTools::MakeTempPivotDialog),
+		FExecuteAction::CreateSP(this, &SControlRigEditModeTools::MakeMotionTrailDialog),
 		NAME_None,
-		LOCTEXT("TempPivot", "Temp Pivot"),
-		LOCTEXT("TempPivotTooltip", "Create a temporary pivot to transform the selected Control"),
-		FSlateIcon(FEditorStyle::GetStyleSetName(), "AnimationEditor.ApplyAnimation"), //MZ todo replace with correct icon
+		LOCTEXT("MotionTrails", "Motion Trails"),
+		LOCTEXT("MotionTrailsTooltip", "Display motion trails for animated objects"),
+		FSlateIcon(TEXT("ControlRigEditorStyle"), TEXT("ControlRig.SnapperTool")),
 		EUserInterfaceActionType::Button
 	);
-	*/
+	//Pivot
+	ToolBarBuilder.AddToolBarButton(
+		FUIAction(
+		FExecuteAction::CreateSP(this, &SControlRigEditModeTools::ToggleEditPivotMode),
+		FCanExecuteAction(),
+		FIsActionChecked::CreateLambda([] {
+				if (FLevelEditorModule* LevelEditorModule = FModuleManager::GetModulePtr<FLevelEditorModule>(TEXT("LevelEditor")))
+				{
+					TSharedPtr<ILevelEditor> LevelEditorPtr = LevelEditorModule->GetLevelEditorInstance().Pin();
+
+					if (LevelEditorPtr.IsValid())
+					{
+						FString ActiveToolName = LevelEditorPtr->GetEditorModeManager().GetInteractiveToolsContext()->ToolManager->GetActiveToolName(EToolSide::Left);
+						if (ActiveToolName == TEXT("SequencerPivotTool"))
+						{
+							return true;
+						}
+					}
+				}
+				return false;
+
+			})
+		),
+		NAME_None,
+		LOCTEXT("TempPivot", "Temp Pivot"),
+		LOCTEXT("TempPivotTooltip", "Create a temporary pivot to rotate the selected Control"),
+		FSlateIcon(TEXT("ControlRigEditorStyle"), TEXT("ControlRig.SnapperTool")),
+		EUserInterfaceActionType::ToggleButton
+		);
+	
 	ToolBarBuilder.AddSeparator();
 }
 
@@ -601,7 +635,6 @@ void SControlRigEditModeTools::MakeTweenDialog()
 	}
 }
 
-
 void SControlRigEditModeTools::MakeSnapperDialog()
 {
 	FControlRigEditMode* ControlRigEditMode = static_cast<FControlRigEditMode*>(ModeTools->GetActiveMode(FControlRigEditMode::ModeName));
@@ -611,14 +644,35 @@ void SControlRigEditModeTools::MakeSnapperDialog()
 	}
 }
 
-
-
-void SControlRigEditModeTools::MakeTempPivotDialog()
+void SControlRigEditModeTools::MakeMotionTrailDialog()
 {
-	FControlRigEditMode* ControlRigEditMode = static_cast<FControlRigEditMode*>(ModeTools->GetActiveMode(FControlRigEditMode::ModeName));
+	FControlRigEditMode* ControlRigEditMode = static_cast<FControlRigEditMode*>(GLevelEditorModeTools().GetActiveMode(FControlRigEditMode::ModeName));
 	if (ControlRigEditMode)
 	{
-		FGlobalTabmanager::Get()->TryInvokeTab(IControlRigEditorModule::ControlRigTempPivotTab);
+		FGlobalTabmanager::Get()->TryInvokeTab(IControlRigEditorModule::ControlRigMotionTrailTab);
+	}
+}
+
+void SControlRigEditModeTools::ToggleEditPivotMode()
+{
+	if (FLevelEditorModule* LevelEditorModule = FModuleManager::GetModulePtr<FLevelEditorModule>(TEXT("LevelEditor")))
+	{
+		TSharedPtr<ILevelEditor> LevelEditorPtr = LevelEditorModule->GetLevelEditorInstance().Pin();
+
+		if (LevelEditorPtr.IsValid())
+		{
+			FString ActiveToolName = LevelEditorPtr->GetEditorModeManager().GetInteractiveToolsContext()->ToolManager->GetActiveToolName(EToolSide::Left);
+			if(ActiveToolName == TEXT("SequencerPivotTool"))
+			{
+				LevelEditorPtr->GetEditorModeManager().GetInteractiveToolsContext()->ToolManager->DeactivateTool(EToolSide::Left, EToolShutdownType::Completed);
+			}
+			else
+			{
+				LevelEditorPtr->GetEditorModeManager().GetInteractiveToolsContext()->ToolManager->SelectActiveToolType(EToolSide::Left, TEXT("SequencerPivotTool"));
+				LevelEditorPtr->GetEditorModeManager().GetInteractiveToolsContext()->ToolManager->ActivateTool(EToolSide::Left);
+
+			}
+		}
 	}
 }
 
diff --git a/Engine/Plugins/Animation/ControlRig/Source/ControlRigEditor/Private/EditMode/SControlRigEditModeTools.h b/Engine/Plugins/Animation/ControlRig/Source/ControlRigEditor/Private/EditMode/SControlRigEditModeTools.h
--- a/Engine/Plugins/Animation/ControlRig/Source/ControlRigEditor/Private/EditMode/SControlRigEditModeTools.h
+++ b/Engine/Plugins/Animation/ControlRig/Source/ControlRigEditor/Private/EditMode/SControlRigEditModeTools.h
@@ -81,7 +81,8 @@ private:
 	void MakePoseDialog();
 	void MakeTweenDialog();
 	void MakeSnapperDialog();
-	void MakeTempPivotDialog();
+	void MakeMotionTrailDialog();
+	void ToggleEditPivotMode();
 
 	//TODO may put back void MakeSelectionSetDialog();
 	//TWeakPtr<SWindow> SelectionSetWindow;
diff --git a/Engine/Plugins/Animation/ControlRig/Source/ControlRigEditor/Private/Sequencer/ControlRigParameterTrackEditor.cpp b/Engine/Plugins/Animation/ControlRig/Source/ControlRigEditor/Private/Sequencer/ControlRigParameterTrackEditor.cpp
--- a/Engine/Plugins/Animation/ControlRig/Source/ControlRigEditor/Private/Sequencer/ControlRigParameterTrackEditor.cpp
+++ b/Engine/Plugins/Animation/ControlRig/Source/ControlRigEditor/Private/Sequencer/ControlRigParameterTrackEditor.cpp
@@ -1387,7 +1387,7 @@ void FControlRigParameterTrackEditor::OnAddTransformKeysForSelectedObjects(EMovi
 				USceneComponent* Component = Cast<USceneComponent>(ObjectBinding->GetBoundObject());
 				if (Component)
 				{
-					AddControlKeys(Component, ControlRig, Name, ControlName, Channel, ESequencerKeyMode::ManualKeyForced, FLT_MAX);
+					AddControlKeys(Component, ControlRig, Name, ControlName, (EControlRigContextChannelToKey) Channel, ESequencerKeyMode::ManualKeyForced, FLT_MAX);
 				}
 			}
 		}
@@ -1981,10 +1981,11 @@ void FControlRigParameterTrackEditor::HandleOnInitialized(UControlRig* ControlRi
 
 void FControlRigParameterTrackEditor::HandleControlModified(UControlRig* ControlRig, FRigControlElement* ControlElement, const FRigControlModifiedContext& Context)
 {
-	if (!GetSequencer().IsValid() || !GetSequencer()->IsAllowedToChange())
+	if (!GetSequencer().IsValid() || !GetSequencer()->IsAllowedToChange() || Context.SetKey == EControlRigSetKey::Never)
 	{
 		return;
 	}
+	FTransform  Transform = ControlRig->GetControlLocalTransform(ControlElement->GetName());	
 	UMovieScene* MovieScene = GetSequencer()->GetFocusedMovieSceneSequence()->GetMovieScene();
 	const TArray<FMovieSceneBinding>& Bindings = MovieScene->GetBindings();
 	for (const FMovieSceneBinding& Binding : Bindings)
@@ -2003,18 +2004,14 @@ void FControlRigParameterTrackEditor::HandleControlModified(UControlRig* Control
 					{
 						KeyMode = ESequencerKeyMode::ManualKeyForced;
 					}
-					else if (Context.SetKey == EControlRigSetKey::Never)
-					{
-						KeyMode = ESequencerKeyMode::ManualKey; //best we have here
-					}
-					AddControlKeys(Component, ControlRig, Name, ControlElement->GetName(), EMovieSceneTransformChannel::All, KeyMode, Context.LocalTime);
+					AddControlKeys(Component, ControlRig, Name, ControlElement->GetName(), (EControlRigContextChannelToKey)Context.KeyMask, KeyMode, Context.LocalTime);
 				}
 			}
 		}
 	}
 }
 
-void FControlRigParameterTrackEditor::GetControlRigKeys(UControlRig* InControlRig, FName ParameterName, EMovieSceneTransformChannel ChannelsToKey, UMovieSceneControlRigParameterSection* SectionToKey, FGeneratedTrackKeys& OutGeneratedKeys)
+void FControlRigParameterTrackEditor::GetControlRigKeys(UControlRig* InControlRig, FName ParameterName, EControlRigContextChannelToKey ChannelsToKey,  ESequencerKeyMode KeyMode, UMovieSceneControlRigParameterSection* SectionToKey,FGeneratedTrackKeys& OutGeneratedKeys)
 {
 	const TArray<bool>& ControlsMask = SectionToKey->GetControlsMask();
 	EMovieSceneTransformChannel TransformMask = SectionToKey->GetTransformMask().GetChannels();
@@ -2022,9 +2019,9 @@ void FControlRigParameterTrackEditor::GetControlRigKeys(UControlRig* InControlRi
 	TArray<FRigControlElement*> Controls;
 	InControlRig->GetControlsInOrder(Controls);
 	// If key all is enabled, for a key on all the channels
-	if (GetSequencer()->GetKeyGroupMode() == EKeyGroupMode::KeyAll)
+	if (KeyMode != ESequencerKeyMode::ManualKeyForced && GetSequencer()->GetKeyGroupMode() == EKeyGroupMode::KeyAll)
 	{
-		ChannelsToKey = EMovieSceneTransformChannel::All;
+		ChannelsToKey = EControlRigContextChannelToKey::AllTransform;
 	}
 
 	//Need seperate index fo bools,ints and enums and floats since there are seperate entries for each later when they are accessed by the set key stuff.
@@ -2120,9 +2117,9 @@ void FControlRigParameterTrackEditor::GetControlRigKeys(UControlRig* InControlRi
 				Scale = Val.GetScale3D();
 			}
 			FVector3f CurrentVector = Translation;
-			bool bKeyX = bSetKey && EnumHasAnyFlags(ChannelsToKey, EMovieSceneTransformChannel::TranslationX);
-			bool bKeyY = bSetKey && EnumHasAnyFlags(ChannelsToKey, EMovieSceneTransformChannel::TranslationY);
-			bool bKeyZ = bSetKey && EnumHasAnyFlags(ChannelsToKey, EMovieSceneTransformChannel::TranslationZ);
+			bool bKeyX = bSetKey && EnumHasAnyFlags(ChannelsToKey, EControlRigContextChannelToKey::TranslationX);
+			bool bKeyY = bSetKey && EnumHasAnyFlags(ChannelsToKey, EControlRigContextChannelToKey::TranslationY);
+			bool bKeyZ = bSetKey && EnumHasAnyFlags(ChannelsToKey, EControlRigContextChannelToKey::TranslationZ);
 			if (GetSequencer()->GetKeyGroupMode() == EKeyGroupMode::KeyGroup && (bKeyX || bKeyY || bKeyZ))
 			{
 				bKeyX = bKeyY = bKeyZ = true;
@@ -2144,9 +2141,9 @@ void FControlRigParameterTrackEditor::GetControlRigKeys(UControlRig* InControlRi
 			OutGeneratedKeys.Add(FMovieSceneChannelValueSetter::Create<FMovieSceneFloatChannel>(ChannelIndex++, CurrentVector.Z, bKeyZ));
 
 			FRotator CurrentRotator = Rotation;
-			bKeyX = bSetKey && EnumHasAnyFlags(ChannelsToKey, EMovieSceneTransformChannel::RotationX);
-			bKeyY = bSetKey && EnumHasAnyFlags(ChannelsToKey, EMovieSceneTransformChannel::RotationY);
-			bKeyZ = bSetKey && EnumHasAnyFlags(ChannelsToKey, EMovieSceneTransformChannel::RotationZ);
+			bKeyX = bSetKey && EnumHasAnyFlags(ChannelsToKey, EControlRigContextChannelToKey::RotationX);
+			bKeyY = bSetKey && EnumHasAnyFlags(ChannelsToKey, EControlRigContextChannelToKey::RotationY);
+			bKeyZ = bSetKey && EnumHasAnyFlags(ChannelsToKey, EControlRigContextChannelToKey::RotationZ);
 			if (GetSequencer()->GetKeyGroupMode() == EKeyGroupMode::KeyGroup && (bKeyX || bKeyY || bKeyZ))
 			{
 				bKeyX = bKeyY = bKeyZ = true;
@@ -2184,9 +2181,9 @@ void FControlRigParameterTrackEditor::GetControlRigKeys(UControlRig* InControlRi
 			if (ControlElement->Settings.ControlType == ERigControlType::Transform || ControlElement->Settings.ControlType == ERigControlType::EulerTransform)
 			{
 				CurrentVector = Scale;
-				bKeyX = bSetKey && EnumHasAnyFlags(ChannelsToKey, EMovieSceneTransformChannel::ScaleX);
-				bKeyY = bSetKey && EnumHasAnyFlags(ChannelsToKey, EMovieSceneTransformChannel::ScaleY);
-				bKeyZ = bSetKey && EnumHasAnyFlags(ChannelsToKey, EMovieSceneTransformChannel::ScaleZ);
+				bKeyX = bSetKey && EnumHasAnyFlags(ChannelsToKey, EControlRigContextChannelToKey::ScaleX);
+				bKeyY = bSetKey && EnumHasAnyFlags(ChannelsToKey, EControlRigContextChannelToKey::ScaleY);
+				bKeyZ = bSetKey && EnumHasAnyFlags(ChannelsToKey, EControlRigContextChannelToKey::ScaleZ);
 				if (GetSequencer()->GetKeyGroupMode() == EKeyGroupMode::KeyGroup && (bKeyX || bKeyY || bKeyZ))
 				{
 					bKeyX = bKeyY = bKeyZ = true;
@@ -2277,7 +2274,7 @@ FKeyPropertyResult FControlRigParameterTrackEditor::AddKeysToControlRigHandle(US
 }
 
 FKeyPropertyResult FControlRigParameterTrackEditor::AddKeysToControlRig(
-	USceneComponent *InSceneComp, UControlRig* InControlRig, FFrameNumber KeyTime, FGeneratedTrackKeys& GeneratedKeys,
+	USceneComponent *InSceneComp, UControlRig* InControlRig, FFrameNumber KeyTime,  FGeneratedTrackKeys& GeneratedKeys,
 	ESequencerKeyMode KeyMode, TSubclassOf<UMovieSceneTrack> TrackClass, FName ControlRigName, FName RigControlName)
 {
 	FKeyPropertyResult KeyPropertyResult;
@@ -2300,7 +2297,7 @@ FKeyPropertyResult FControlRigParameterTrackEditor::AddKeysToControlRig(
 	return KeyPropertyResult;
 }
 
-void FControlRigParameterTrackEditor::AddControlKeys(USceneComponent *InSceneComp, UControlRig* InControlRig, FName ControlRigName, FName RigControlName, EMovieSceneTransformChannel ChannelsToKey, ESequencerKeyMode KeyMode, float InLocalTime)
+void FControlRigParameterTrackEditor::AddControlKeys(USceneComponent *InSceneComp, UControlRig* InControlRig, FName ControlRigName, FName RigControlName, EControlRigContextChannelToKey ChannelsToKey, ESequencerKeyMode KeyMode, float InLocalTime)
 {
 	if (KeyMode == ESequencerKeyMode::ManualKey || (GetSequencer().IsValid() && !GetSequencer()->IsAllowedToChange()))
 	{
@@ -2336,7 +2333,7 @@ void FControlRigParameterTrackEditor::AddControlKeys(USceneComponent *InSceneCom
 
 	TSharedRef<FGeneratedTrackKeys> GeneratedKeys = MakeShared<FGeneratedTrackKeys>();
 
-	GetControlRigKeys(InControlRig, RigControlName, ChannelsToKey, ParamSection, *GeneratedKeys);
+	GetControlRigKeys(InControlRig, RigControlName, ChannelsToKey, KeyMode, ParamSection, *GeneratedKeys);
 	TGuardValue<bool> Guard(bIsDoingSelection, true);
 
 	auto OnKeyProperty = [=](FFrameNumber Time) -> FKeyPropertyResult
diff --git a/Engine/Plugins/Animation/ControlRig/Source/ControlRigEditor/Private/Sequencer/ControlRigParameterTrackEditor.h b/Engine/Plugins/Animation/ControlRig/Source/ControlRigEditor/Private/Sequencer/ControlRigParameterTrackEditor.h
--- a/Engine/Plugins/Animation/ControlRig/Source/ControlRigEditor/Private/Sequencer/ControlRigParameterTrackEditor.h
+++ b/Engine/Plugins/Animation/ControlRig/Source/ControlRigEditor/Private/Sequencer/ControlRigParameterTrackEditor.h
@@ -25,6 +25,7 @@ class UMovieSceneControlRigParameterSection;
 class UFKControlRig;
 class FEditorModeTools;
 class FControlRigEditMode;
+struct FRigControlModifiedContext;
 
 /**
  * Tools for animation tracks
@@ -151,8 +152,8 @@ private:
 
 public:
 
-	void AddControlKeys(USceneComponent *InSceneComp, UControlRig* InControlRig, FName PropertyName, FName ParameterName, EMovieSceneTransformChannel ChannelsToKey, ESequencerKeyMode KeyMode, float InLocalTime);
-	void GetControlRigKeys(UControlRig* InControlRig, FName ParameterName, EMovieSceneTransformChannel ChannelsToKey, UMovieSceneControlRigParameterSection* SectionToKey, FGeneratedTrackKeys& OutGeneratedKeys);
+	void AddControlKeys(USceneComponent *InSceneComp, UControlRig* InControlRig, FName PropertyName, FName ParameterName, EControlRigContextChannelToKey ChannelsToKey, ESequencerKeyMode KeyMode, float InLocalTime);
+	void GetControlRigKeys(UControlRig* InControlRig, FName ParameterName, EControlRigContextChannelToKey ChannelsToKey, ESequencerKeyMode KeyMode, UMovieSceneControlRigParameterSection* SectionToKey, FGeneratedTrackKeys& OutGeneratedKeys);
 	FKeyPropertyResult AddKeysToControlRig(
 		USceneComponent *InSceneComp, UControlRig* InControlRig, FFrameNumber KeyTime, FGeneratedTrackKeys& GeneratedKeys,
 		ESequencerKeyMode KeyMode, TSubclassOf<UMovieSceneTrack> TrackClass, FName ControlRigName, FName RigControlName);
diff --git a/Engine/Plugins/Animation/ControlRig/Source/ControlRigEditor/Private/Tools/ControlRigSnapper.cpp b/Engine/Plugins/Animation/ControlRig/Source/ControlRigEditor/Private/Tools/ControlRigSnapper.cpp
--- a/Engine/Plugins/Animation/ControlRig/Source/ControlRigEditor/Private/Tools/ControlRigSnapper.cpp
+++ b/Engine/Plugins/Animation/ControlRig/Source/ControlRigEditor/Private/Tools/ControlRigSnapper.cpp
@@ -182,30 +182,6 @@ bool FControlRigSnapper::GetControlRigControlTransforms(UWorld* World,ULevelSequ
 	return false;
 }
 
-//Matinee version of this doesn't actually set the key it only adds..sigh
-static void SetOrAddKey(TMovieSceneChannelData<FMovieSceneFloatValue>& ChannelData, FFrameNumber Time, float Value)
-{
-	int32 ExistingIndex = ChannelData.FindKey(Time);
-	if (ExistingIndex != INDEX_NONE)
-	{
-		FMovieSceneFloatValue& FloatValue = ChannelData.GetValues()[ExistingIndex]; //-V758
-		FloatValue.Value = Value;
-	}
-	else
-	{
-		FMovieSceneFloatValue NewKey(Value);
-		ERichCurveTangentWeightMode WeightedMode = RCTWM_WeightedNone;
-		NewKey.InterpMode = ERichCurveInterpMode::RCIM_Cubic;
-		NewKey.TangentMode = ERichCurveTangentMode::RCTM_Auto;
-		NewKey.Tangent.ArriveTangent = 0.0f;
-		NewKey.Tangent.LeaveTangent = 0.0f;
-		NewKey.Tangent.TangentWeightMode = WeightedMode;
-		NewKey.Tangent.ArriveTangentWeight = 0.0f;
-		NewKey.Tangent.LeaveTangentWeight = 0.0f;
-		ChannelData.AddKey(Time, NewKey);
-	}
-}
-
 struct FGuidAndActor
 {
 	FGuidAndActor(FGuid InGuid, AActor* InActor) : Guid(InGuid), Actor(InActor) {};
@@ -322,30 +298,30 @@ struct FGuidAndActor
 			if (SnapSettings->bSnapPosition)
 			{
 				TMovieSceneChannelData<FMovieSceneFloatValue> ChannelData = Channels[0]->GetData();
-				SetOrAddKey(ChannelData, Frame, Location.X);
+				MovieSceneToolHelpers::SetOrAddKey(ChannelData, Frame, Location.X);
 				ChannelData = Channels[1]->GetData();
-				SetOrAddKey(ChannelData, Frame, Location.Y);
+				MovieSceneToolHelpers::SetOrAddKey(ChannelData, Frame, Location.Y);
 				ChannelData = Channels[2]->GetData();
-				SetOrAddKey(ChannelData, Frame, Location.Z);
+				MovieSceneToolHelpers::SetOrAddKey(ChannelData, Frame, Location.Z);
 			}
 			if (SnapSettings->bSnapRotation)
 			{
 				//todo winding
 				TMovieSceneChannelData<FMovieSceneFloatValue> ChannelData = Channels[3]->GetData();
-				SetOrAddKey(ChannelData, Frame, Rotation.Yaw);
+				MovieSceneToolHelpers::SetOrAddKey(ChannelData, Frame, Rotation.Yaw);
 				ChannelData = Channels[4]->GetData();
-				SetOrAddKey(ChannelData, Frame, Rotation.Pitch);
+				MovieSceneToolHelpers::SetOrAddKey(ChannelData, Frame, Rotation.Pitch);
 				ChannelData = Channels[5]->GetData();
-				SetOrAddKey(ChannelData, Frame, Rotation.Roll);
+				MovieSceneToolHelpers::SetOrAddKey(ChannelData, Frame, Rotation.Roll);
 			}
 			if (SnapSettings->bSnapScale)
 			{
 				TMovieSceneChannelData<FMovieSceneFloatValue> ChannelData = Channels[6]->GetData();
-				SetOrAddKey(ChannelData, Frame, Scale3D.X);
+				MovieSceneToolHelpers::SetOrAddKey(ChannelData, Frame, Scale3D.X);
 				ChannelData = Channels[7]->GetData();
-				SetOrAddKey(ChannelData, Frame, Scale3D.X);
+				MovieSceneToolHelpers::SetOrAddKey(ChannelData, Frame, Scale3D.X);
 				ChannelData = Channels[8]->GetData();
-				SetOrAddKey(ChannelData, Frame, Scale3D.X);
+				MovieSceneToolHelpers::SetOrAddKey(ChannelData, Frame, Scale3D.X);
 			}		
 		}
 		//now we need to set auto tangents
diff --git a/Engine/Plugins/Animation/ControlRig/Source/ControlRigEditor/Public/IControlRigEditorModule.h b/Engine/Plugins/Animation/ControlRig/Source/ControlRigEditor/Public/IControlRigEditorModule.h
--- a/Engine/Plugins/Animation/ControlRig/Source/ControlRigEditor/Public/IControlRigEditorModule.h
+++ b/Engine/Plugins/Animation/ControlRig/Source/ControlRigEditor/Public/IControlRigEditorModule.h
@@ -49,5 +49,5 @@ public:
 	static const FName ControlRigPoseTab;
 	static const FName ControlRigTweenTab;
 	static const FName ControlRigSnapperTab;
-	static const FName ControlRigTempPivotTab;
+	static const FName ControlRigMotionTrailTab;
 };
diff --git a/Engine/Plugins/Editor/SequencerAnimTools/SequencerAnimTools.uplugin b/Engine/Plugins/Editor/SequencerAnimTools/SequencerAnimTools.uplugin
new file mode 100644
--- /dev/null
+++ b/Engine/Plugins/Editor/SequencerAnimTools/SequencerAnimTools.uplugin
@@ -0,0 +1,38 @@
+{
+	"FileVersion": 3,
+	"Version": 1,
+	"VersionName": "1.0",
+	"FriendlyName": "Sequencer Anim Tools",
+	"Description": "Animation Tools For Sequencer and Control Rig",
+	"Category": "Editor",
+	"CreatedBy": "Epic Games, Inc",
+	"CreatedByURL": "http://epicgames.com",
+	"DocsURL": "",
+	"MarketplaceURL": "",
+	"SupportURL": "",
+	"CanContainContent": false,
+	"IsBetaVersion": false,
+	"IsExperimentalVersion": false,
+	"Installed": false,
+	"Plugins": [
+		{
+			"Name": "ControlRig",
+			"Enabled": true
+		},
+		{
+			"Name": "LevelSequenceEditor",
+			"Enabled": true
+		},
+		{
+			"Name": "SequencerScripting",
+			"Enabled": true
+		}
+	],
+	"Modules": [
+		{
+			"Name": "SequencerAnimTools",
+			"Type": "Editor",
+			"LoadingPhase": "Default"
+		}
+	]
+}
\ No newline at end of file
diff --git a/Engine/Plugins/Editor/SequencerAnimTools/Source/SequencerAnimTools/Private/BaseSequencerAnimTool.cpp b/Engine/Plugins/Editor/SequencerAnimTools/Source/SequencerAnimTools/Private/BaseSequencerAnimTool.cpp
new file mode 100644
--- /dev/null
+++ b/Engine/Plugins/Editor/SequencerAnimTools/Source/SequencerAnimTools/Private/BaseSequencerAnimTool.cpp
@@ -0,0 +1,7 @@
+// Copyright Epic Games, Inc. All Rights Reserved.
+
+#include "BaseSequencerAnimTool.h"
+
+UBaseSequencerAnimTool::UBaseSequencerAnimTool(const FObjectInitializer& ObjectInitializer)
+	: Super(ObjectInitializer)
+{}
\ No newline at end of file
diff --git a/Engine/Plugins/Editor/SequencerAnimTools/Source/SequencerAnimTools/Private/EditPivotTool.cpp b/Engine/Plugins/Editor/SequencerAnimTools/Source/SequencerAnimTools/Private/EditPivotTool.cpp
new file mode 100644
--- /dev/null
+++ b/Engine/Plugins/Editor/SequencerAnimTools/Source/SequencerAnimTools/Private/EditPivotTool.cpp
@@ -0,0 +1,502 @@
+// Copyright Epic Games, Inc. All Rights Reserved.
+
+#include "EditPivotTool.h"
+#include "LevelEditorSequencerIntegration.h"
+#include "ISequencer.h"
+#include "Sequencer/SequencerTrailHierarchy.h"
+#include "Sequencer/MovieSceneTransformTrail.h"
+#include "BaseGizmos/TransformGizmo.h"
+#include "InteractiveGizmoManager.h"
+#include "BaseBehaviors/SingleClickBehavior.h"
+#include "BaseBehaviors/SingleKeyCaptureBehavior.h"
+#include "BaseGizmos/TransformGizmoUtil.h"
+#include "Framework/Commands/GenericCommands.h"
+#include "LevelEditor.h"
+
+#include "InteractiveToolManager.h"
+#include "ToolBuilderUtil.h"
+// for raycast into World
+#include "CollisionQueryParams.h"
+#include "Engine/World.h"
+
+#include "SceneManagement.h"
+#include "ScopedTransaction.h"
+#include "Tools/MotionTrailOptions.h"
+#include "EditorModeManager.h"
+#include "EditorViewportClient.h"
+#include "EditorModeManager.h"
+#include "ControlRig.h"
+#include "ControlRigSequencerEditorLibrary.h"
+#include "LevelSequence.h"
+#include "LevelSequenceEditorBlueprintLibrary.h"
+#include "ILevelSequenceEditorToolkit.h"
+#include "EditorModeManager.h"
+#include "InteractiveToolManager.h"
+#include "EdModeInteractiveToolsContext.h"
+
+#define LOCTEXT_NAMESPACE "SequencerAnimTools"
+
+
+
+void FEditPivotCommands::RegisterCommands()
+{
+	UI_COMMAND(ResetPivot, "Reset Pivot To Original", "Reset pivot back to original location", EUserInterfaceActionType::Button, FInputChord(EModifierKey::Control|EModifierKey::Shift|EModifierKey::Alt, EKeys::G));
+}
+
+
+TMap<TWeakObjectPtr<UControlRig>, FControlRigMappings> USequencerPivotTool::SavedPivotLocations;
+TArray<FControlRigMappings> USequencerPivotTool::LastSelectedObjects;
+
+static void GetControlRigsAndSequencer(TArray<UControlRig*>& ControlRigs, TWeakPtr<ISequencer>& SequencerPtr, ULevelSequence** LevelSequence)
+{
+	*LevelSequence = ULevelSequenceEditorBlueprintLibrary::GetCurrentLevelSequence();
+	if (*LevelSequence)
+	{
+		IAssetEditorInstance* AssetEditor = GEditor->GetEditorSubsystem<UAssetEditorSubsystem>()->FindEditorForAsset(*LevelSequence, false);
+		ILevelSequenceEditorToolkit* LevelSequenceEditor = static_cast<ILevelSequenceEditorToolkit*>(AssetEditor);
+		SequencerPtr = LevelSequenceEditor ? LevelSequenceEditor->GetSequencer() : nullptr;
+		if (SequencerPtr.IsValid())
+		{
+			ControlRigs = UControlRigSequencerEditorLibrary::GetVisibleControlRigs();
+		}
+	}
+}
+bool USequencerPivotToolBuilder::CanBuildTool(const FToolBuilderState& SceneState) const
+{
+	//only start if we have anything selected  or we selected somethign last time
+	//which will then get selected later once the tool starts.
+	if (USequencerPivotTool::LastSelectedObjects.Num() > 0)
+	{
+		return true;
+	}
+	ULevelSequence* LevelSequence;
+	TArray<UControlRig*> ControlRigs;
+	TWeakPtr<ISequencer> SequencerPtr;
+	GetControlRigsAndSequencer(ControlRigs, SequencerPtr, &LevelSequence);
+
+	for (UControlRig* ControlRig : ControlRigs)
+	{
+		if (ControlRig && ControlRig->CurrentControlSelection().Num() > 0)
+		{
+			return true;
+		}
+	}
+	
+	
+	return false;
+}
+
+UInteractiveTool* USequencerPivotToolBuilder::BuildTool(const FToolBuilderState& SceneState) const
+{
+	USequencerPivotTool* NewTool = NewObject<USequencerPivotTool>(SceneState.ToolManager);
+	NewTool->SetWorld(SceneState.World, SceneState.GizmoManager);
+	return NewTool;
+}
+/*
+* Proxy
+*/
+
+struct HSequencerPivotProxy : public HHitProxy
+{
+	DECLARE_HIT_PROXY();
+
+	HSequencerPivotProxy(EHitProxyPriority InPriority = HPP_Foreground) :
+		HHitProxy(InPriority)
+	{}
+};
+
+IMPLEMENT_HIT_PROXY(HSequencerPivotProxy, HHitProxy);
+
+/*
+* Tool
+*/
+
+void USequencerPivotTool::SetWorld(UWorld* World, UInteractiveGizmoManager* InGizmoManager)
+{
+	TargetWorld = World;
+	GizmoManager = InGizmoManager;
+}
+
+void USequencerPivotTool::Setup()
+{
+	//when entered we check to see if shift is pressed this can change where we set the pivot on start or reset
+	FModifierKeysState KeyState = FSlateApplication::Get().GetModifierKeys();
+	bShiftPressedWhenStarted = KeyState.IsShiftDown();
+
+	UInteractiveTool::Setup();
+
+	ClickBehavior = NewObject<USingleClickInputBehavior>(this);
+	ClickBehavior->Initialize(this);
+	AddInputBehavior(ClickBehavior);
+
+	TransformProxy = NewObject<UTransformProxy>(this);
+
+	FString GizmoIdentifier = TEXT("PivotToolGizmoIdentifier");
+	TransformGizmo = GizmoManager->Create3AxisTransformGizmo(this, GizmoIdentifier);
+
+	TransformProxy->OnTransformChanged.AddUObject(this, &USequencerPivotTool::GizmoTransformChanged);
+	TransformProxy->OnBeginTransformEdit.AddUObject(this, &USequencerPivotTool::GizmoTransformStarted);
+	TransformProxy->OnEndTransformEdit.AddUObject(this, &USequencerPivotTool::GizmoTransformEnded);
+	TransformGizmo->SetActiveTarget(TransformProxy, GetToolManager());
+
+	GetControlRigsAndSequencer(ControlRigs, SequencerPtr, &LevelSequence);
+	UpdateTransformAndSelectionOnEntering();
+	UpdateGizmoTransform();
+	UpdateGizmoVisibility();
+	//we get delegates last since we may select something above
+	for (UControlRig* ControlRig : ControlRigs)
+	{
+		ControlRig->ControlSelected().AddUObject(this,&USequencerPivotTool::HandleControlSelected);
+	}
+	SaveLastSelected();
+
+	CommandBindings = MakeShareable(new FUICommandList);
+
+	FEditPivotCommands::Register();
+
+	const FEditPivotCommands& Commands = FEditPivotCommands::Get();
+
+	CommandBindings->MapAction(
+		Commands.ResetPivot,
+		FExecuteAction::CreateUObject(this, &USequencerPivotTool::ResetPivot)
+		);
+
+}
+
+void USequencerPivotTool::ResetPivot()
+{
+	SetGizmoBasedOnSelection(false);
+	UpdateGizmoTransform();
+}
+
+//Last selected is really the ones that were selected when you entered the tool
+//we use this in case nothing was selected when the tool is active so that we instead select it.
+void USequencerPivotTool::SaveLastSelected()
+{
+	LastSelectedObjects.SetNum(0);
+	for (UControlRig* ControlRig : ControlRigs)
+	{
+		TArray<FName> SelectedControls = ControlRig->CurrentControlSelection();
+		for (const FName& Name : SelectedControls)
+		{
+			FControlRigMappings Mapping;
+			Mapping.ControlRig = ControlRig;
+			Mapping.PivotTransforms.Add(Name, GizmoTransform);
+			LastSelectedObjects.Add(Mapping);
+		}
+	}
+}
+
+//when we enter the tool if we have things selected we get it's last transform position
+//if not we should have the last thing selected, and if so we select that.
+void USequencerPivotTool::UpdateTransformAndSelectionOnEntering()
+{
+	//if shift was pressed when we started we don't use saved, this will move pivot to last object
+	bool bhaveSomethingSelected = SetGizmoBasedOnSelection(!bShiftPressedWhenStarted);
+	//okay nothing selected we select the last thing selec
+	if (bhaveSomethingSelected == false)
+	{
+		for (FControlRigMappings& Mappings : LastSelectedObjects)
+		{
+			if (Mappings.ControlRig.IsValid())
+			{
+				for (TPair<FName, FTransform>& Pair: Mappings.PivotTransforms)
+				{
+					Mappings.ControlRig->SelectControl(Pair.Key, true);
+				}
+			}
+		}
+		SetGizmoBasedOnSelection(!bShiftPressedWhenStarted);
+	}
+}
+
+bool USequencerPivotTool::ProcessCommandBindings(const FKey Key, const bool bRepeat) const
+{
+	if (CommandBindings.IsValid())
+	{
+		FModifierKeysState KeyState = FSlateApplication::Get().GetModifierKeys();
+		return CommandBindings->ProcessCommandBindings(Key, KeyState, bRepeat);
+	}
+	return false;
+}
+
+bool USequencerPivotTool::SetGizmoBasedOnSelection(bool bUseSaved)
+{
+
+	GizmoTransform = FTransform::Identity;
+	FVector Location(0.0f, 0.0f, 0.0f);
+	ISequencer* Sequencer = SequencerPtr.Pin().Get();
+	if (Sequencer == nullptr)
+	{
+		return false;
+	}
+	const FFrameRate TickResolution = Sequencer->GetFocusedTickResolution();
+	const FFrameTime FrameTime = Sequencer->GetLocalTime().ConvertTo(TickResolution);
+
+	bool bHaveSomethingSelected = false;
+	for (UControlRig* ControlRig : ControlRigs)
+	{
+		FControlRigMappings* Mappings = SavedPivotLocations.Find(ControlRig);
+		TArray<FName> SelectedControls = ControlRig->CurrentControlSelection();
+
+		for (const FName& Name : SelectedControls)
+		{
+			bHaveSomethingSelected = true;
+			bool bHasSavedPivot = false;
+			if (bUseSaved && Mappings)
+			{
+				if (FTransform* Transform = Mappings->PivotTransforms.Find(Name))
+				{
+					GizmoTransform = *Transform;
+					bHasSavedPivot = true;
+				}
+			}
+			if (bHasSavedPivot == false)
+			{
+				GizmoTransform = UControlRigSequencerEditorLibrary::GetControlRigWorldTransform(LevelSequence, ControlRig, Name, FrameTime.RoundToFrame(),
+					ESequenceTimeUnit::TickResolution);
+			}
+		}
+	}
+	GizmoTransform.SetScale3D(FVector(1.0f, 1.0f, 1.0f));
+	GizmoTransform.SetRotation(FQuat::Identity);
+	return bHaveSomethingSelected;
+}
+void USequencerPivotTool::Shutdown(EToolShutdownType ShutdownType)
+{
+	GizmoManager->DestroyAllGizmosByOwner(this);
+	RemoveControlRigDelegates();
+}
+
+void USequencerPivotTool::HandleControlSelected(UControlRig* Subject, FRigControlElement* InControl, bool bSelected)
+{
+
+	if (FLevelEditorModule* LevelEditorModule = FModuleManager::GetModulePtr<FLevelEditorModule>(TEXT("LevelEditor")))
+	{
+		TSharedPtr<ILevelEditor> LevelEditorPtr = LevelEditorModule->GetLevelEditorInstance().Pin();
+
+		if (LevelEditorPtr.IsValid())
+		{
+			FString ActiveToolName = LevelEditorPtr->GetEditorModeManager().GetInteractiveToolsContext()->ToolManager->GetActiveToolName(EToolSide::Left);
+			if (ActiveToolName == TEXT("SequencerPivotTool"))
+			{
+				LevelEditorPtr->GetEditorModeManager().GetInteractiveToolsContext()->ToolManager->DeactivateTool(EToolSide::Left, EToolShutdownType::Completed);
+			}
+		}
+	}
+	
+}
+
+void USequencerPivotTool::RemoveControlRigDelegates()
+{
+	for (UControlRig* ControlRig : ControlRigs)
+	{
+		ControlRig->ControlSelected().RemoveAll(this);
+	}
+}
+
+
+void USequencerPivotTool::GizmoTransformStarted(UTransformProxy* Proxy)
+{
+	ControlRigDrags.SetNum(0);
+	TransactionIndex = INDEX_NONE;
+	ISequencer* Sequencer = SequencerPtr.Pin().Get();
+	if (Sequencer && LevelSequence)
+	{
+		//TransactionIndex = GEditor->BeginTransaction(nullptr, LOCTEXT("TransformPivot", "Transform Pivot"), nullptr);
+		const FFrameRate TickResolution = Sequencer->GetFocusedTickResolution();
+		const FFrameTime FrameTime = Sequencer->GetLocalTime().ConvertTo(TickResolution);
+		const FFrameNumber FrameNumber = FrameTime.RoundToFrame();
+
+		for (UControlRig* ControlRig : ControlRigs)
+		{
+			TArray<FName> SelectedControls = ControlRig->CurrentControlSelection();
+			if (SelectedControls.Num() > 0)
+			{
+				ControlRig->Modify();
+				for (const FName& Name : SelectedControls)
+				{
+					FTransform Transform = UControlRigSequencerEditorLibrary::GetControlRigWorldTransform(LevelSequence, ControlRig, Name, FrameNumber,
+						ESequenceTimeUnit::TickResolution);
+
+					FControlRigSelectionDuringDrag ControlDrag;
+					ControlDrag.LevelSequence = LevelSequence;
+					ControlDrag.ControlName = Name;
+					ControlDrag.ControlRig = ControlRig;
+					ControlDrag.CurrentFrame = FrameNumber;
+					ControlDrag.CurrentTransform = Transform;
+					ControlRigDrags.Add(ControlDrag);
+				}
+			}
+		}
+	}
+
+	StartDragTransform = Proxy->GetTransform();
+	bGizmoBeingDragged = true;
+	bManipulatorMadeChange = false;
+	GizmoTransform = StartDragTransform;
+}
+
+void USequencerPivotTool::GizmoTransformChanged(UTransformProxy* Proxy, FTransform Transform)
+{
+	if (!bGizmoBeingDragged)
+	{
+		return;
+	}
+	GizmoTransform = Transform;
+	FTransform Diff = Transform.GetRelativeTransform(StartDragTransform);
+	if (Diff.GetRotation().IsIdentity(1e-4f) == false)
+	{
+		const bool bSetKey = false;
+		bManipulatorMadeChange = true;
+		for (FControlRigSelectionDuringDrag& ControlDrag : ControlRigDrags)
+		{
+			FVector LocDiff = ControlDrag.CurrentTransform.GetLocation() - Transform.GetLocation();
+			if (LocDiff.IsNearlyZero(1e-4f) == false)
+			{
+				FVector RotatedDiff = Diff.GetRotation().RotateVector(LocDiff);
+				FVector NewLocation = Transform.GetLocation() + RotatedDiff;
+				ControlDrag.CurrentTransform.SetLocation(NewLocation);
+				UControlRigSequencerEditorLibrary::SetControlRigWorldTransform(ControlDrag.LevelSequence, ControlDrag.ControlRig, ControlDrag.ControlName, 
+					ControlDrag.CurrentFrame,ControlDrag.CurrentTransform, ESequenceTimeUnit::TickResolution, bSetKey);
+			}
+		}
+	}
+	
+	StartDragTransform = Transform;
+	UpdateGizmoTransform();
+
+}
+
+void USequencerPivotTool::GizmoTransformEnded(UTransformProxy* Proxy)
+{
+
+	if (bManipulatorMadeChange && TransactionIndex != INDEX_NONE)
+	{
+		//GEditor->EndTransaction();
+	}
+	else if (TransactionIndex != INDEX_NONE)
+	{
+	//	GEditor->CancelTransaction(TransactionIndex);
+	}
+	bGizmoBeingDragged = false;
+	bManipulatorMadeChange = false;
+	UpdateGizmoTransform();
+	SavePivotTransforms();
+}
+
+void USequencerPivotTool::SavePivotTransforms()
+{
+	for (FControlRigMappings& LastObject : LastSelectedObjects)
+	{
+		FControlRigMappings& Mappings = SavedPivotLocations.FindOrAdd(LastObject.ControlRig);
+		Mappings.ControlRig = LastObject.ControlRig;
+		for (TPair<FName, FTransform>& Pair : LastObject.PivotTransforms)
+		{
+			FTransform& Transform = Mappings.PivotTransforms.FindOrAdd(Pair.Key);
+			Transform = GizmoTransform;
+		}
+	}
+}
+
+//currently always visibile
+void USequencerPivotTool::UpdateGizmoVisibility()
+{
+	if (TransformGizmo)
+	{
+		TransformGizmo->SetVisibility(true);
+	}
+}
+
+void USequencerPivotTool::UpdateGizmoTransform()
+{
+	if (!TransformGizmo)
+	{
+		return;
+	}
+
+	TransformGizmo->ReinitializeGizmoTransform(GizmoTransform);
+
+}
+
+FInputRayHit USequencerPivotTool::IsHitByClick(const FInputDeviceRay& ClickPos)
+{
+	//OnUpdateModifierState is not called yet so use modifier keys unfortunately
+	FModifierKeysState KeyState = FSlateApplication::Get().GetModifierKeys();
+	bPickingPivotLocation = KeyState.IsControlDown();
+	if (bPickingPivotLocation == true)
+	{
+		FVector Temp;
+		FInputRayHit Result = FindRayHit(ClickPos.WorldRay, Temp);
+		return Result;
+	}
+	return FInputRayHit();
+}
+
+void USequencerPivotTool::OnClicked(const FInputDeviceRay& ClickPos)
+{
+	if (bPickingPivotLocation)
+	{
+		FVector HitLocation;
+		FInputRayHit Result = FindRayHit(ClickPos.WorldRay, HitLocation);
+		if (Result.bHit)
+		{
+			GizmoTransform.SetLocation(HitLocation);
+			UpdateGizmoTransform();
+		}
+	}
+}
+
+FInputRayHit USequencerPivotTool::FindRayHit(const FRay& WorldRay, FVector& HitPos)
+{
+	// trace a ray into the World
+	FCollisionObjectQueryParams QueryParams(FCollisionObjectQueryParams::AllObjects);
+	FHitResult Result;
+	bool bHitWorld = TargetWorld->LineTraceSingleByObjectType(Result, WorldRay.Origin, WorldRay.PointAt(999999), QueryParams);
+	if (bHitWorld)
+	{
+		HitPos = Result.ImpactPoint;
+		return FInputRayHit(Result.Distance);
+	}
+	return FInputRayHit();
+}
+
+
+
+void USequencerPivotTool::OnPropertyModified(UObject* PropertySet, FProperty* Property)
+{
+	// if the user updated any of the property fields, update the distance
+	//UpdateDistance();
+}
+
+#include "HitProxies.h"
+
+
+void USequencerPivotTool::Render(IToolsContextRenderAPI* RenderAPI)
+{
+
+	if (bPickingPivotLocation == false )
+	{
+		FPrimitiveDrawInterface* PDI = RenderAPI->GetPrimitiveDrawInterface();
+
+		const bool bHitTesting = PDI && PDI->IsHitTesting();
+		const float KeySize = 20.0f;
+		const FLinearColor Color = FLinearColor(1.0f, 0.0f, 0.0f);
+		if (bHitTesting)
+		{
+			PDI->SetHitProxy(new HSequencerPivotProxy());
+		}
+
+		PDI->DrawPoint(GizmoTransform.GetLocation(), Color, KeySize, SDPG_MAX);
+
+		if (bHitTesting)
+		{
+			PDI->SetHitProxy(nullptr);
+		}
+	}
+}
+
+
+
+
+#undef LOCTEXT_NAMESPACE
diff --git a/Engine/Plugins/Editor/SequencerAnimTools/Source/SequencerAnimTools/Private/MotionTrailTool.cpp b/Engine/Plugins/Editor/SequencerAnimTools/Source/SequencerAnimTools/Private/MotionTrailTool.cpp
new file mode 100644
--- /dev/null
+++ b/Engine/Plugins/Editor/SequencerAnimTools/Source/SequencerAnimTools/Private/MotionTrailTool.cpp
@@ -0,0 +1,470 @@
+// Copyright Epic Games, Inc. All Rights Reserved.
+
+#include "MotionTrailTool.h"
+#include "LevelEditorSequencerIntegration.h"
+#include "ISequencer.h"
+#include "Sequencer/SequencerTrailHierarchy.h"
+#include "Sequencer/MovieSceneTransformTrail.h"
+#include "BaseGizmos/TransformGizmo.h"
+#include "InteractiveGizmoManager.h"
+#include "BaseBehaviors/SingleClickBehavior.h"
+#include "BaseBehaviors/SingleKeyCaptureBehavior.h"
+#include "BaseBehaviors/ClickDragBehavior.h"
+#include "BaseGizmos/TransformGizmoUtil.h"
+#include "Framework/Commands/GenericCommands.h"
+
+#include "InteractiveToolManager.h"
+#include "ToolBuilderUtil.h"
+// for raycast into World
+#include "CollisionQueryParams.h"
+#include "Engine/World.h"
+
+#include "SceneManagement.h"
+#include "ScopedTransaction.h"
+#include "Tools/MotionTrailOptions.h"
+#include "EditorModeManager.h"
+#include "EditorViewportClient.h"
+
+#define LOCTEXT_NAMESPACE "SequencerAnimTools"
+
+UInteractiveTool* UMotionTrailToolBuilder::BuildTool(const FToolBuilderState& SceneState) const
+{
+	UMotionTrailTool* NewTool = NewObject<UMotionTrailTool>(SceneState.ToolManager);
+	UEdMode* EdMode = GetTypedOuter<UEdMode>();
+	FEditorModeTools* ModeManager = nullptr;
+	if (EdMode)
+	{
+		ModeManager = EdMode->GetModeManager();
+	}
+	NewTool->SetWorldGizmoModeManager(SceneState.World, SceneState.GizmoManager, ModeManager);
+
+	return NewTool;
+}
+
+bool UMotionTrailTool::ProcessCommandBindings(const FKey Key, const bool bRepeat) const
+{
+	if (CommandBindings.IsValid())
+	{
+		FModifierKeysState KeyState = FSlateApplication::Get().GetModifierKeys();
+		return CommandBindings->ProcessCommandBindings(Key, KeyState, bRepeat);
+	}
+	return false;
+}
+
+void FMotionTrailCommands::RegisterCommands()
+{
+	UI_COMMAND(TranslateSelectedKeysLeft, "Translate Selected Keys Left", "Translate selected keys one frame to the left", EUserInterfaceActionType::Button, FInputChord(EModifierKey::Control, EKeys::Left));
+	UI_COMMAND(TranslateSelectedKeysRight, "Translate Selected Keys Right", "Translate selected keys one frame to the right", EUserInterfaceActionType::Button, FInputChord(EModifierKey::Control, EKeys::Right));
+	UI_COMMAND(FrameSelection, "Frame Selection", "Frame camera around current key selection", EUserInterfaceActionType::Button, FInputChord(EKeys::F));
+}
+
+FString UMotionTrailTool::TrailKeyTransformGizmoInstanceIdentifier = TEXT("TrailKeyTransformGizmoInstanceIdentifier");
+
+
+void UMotionTrailTool::Setup()
+{
+	UInteractiveTool::Setup();
+
+	LeftTarget.Target = this;
+	LeftTarget.bIsLeft = true;
+	RightTarget.Target = this;
+	RightTarget.bIsLeft = false;
+
+	LeftClickBehavior = NewObject<USingleClickInputBehavior>(this);
+	LeftClickBehavior->SetUseLeftMouseButton(); //default
+	LeftClickBehavior->Initialize(&LeftTarget);
+	LeftClickBehavior->Modifiers.RegisterModifier(ShiftModifierId, FInputDeviceState::IsShiftKeyDown);
+	LeftClickBehavior->Modifiers.RegisterModifier(CtrlModifierId, FInputDeviceState::IsCtrlKeyDown);
+	LeftClickBehavior->Modifiers.RegisterModifier(AltModiferId, FInputDeviceState::IsAltKeyDown);
+	AddInputBehavior(LeftClickBehavior);
+
+	RightClickBehavior = NewObject<USingleClickInputBehavior>(this);
+	RightClickBehavior->SetUseRightMouseButton();
+	RightClickBehavior->Initialize(&RightTarget);
+	RightClickBehavior->Modifiers.RegisterModifier(ShiftModifierId, FInputDeviceState::IsShiftKeyDown);
+	RightClickBehavior->Modifiers.RegisterModifier(CtrlModifierId, FInputDeviceState::IsCtrlKeyDown);
+	RightClickBehavior->Modifiers.RegisterModifier(AltModiferId, FInputDeviceState::IsAltKeyDown);
+	AddInputBehavior(RightClickBehavior);
+
+	TransformProxy = NewObject<UTransformProxy>(this);
+	//TransformGizmo = GetToolManager()->GetPairedGizmoManager()->Create3AxisTransformGizmo(this, UMotionTrailTool::TrailKeyTransformGizmoInstanceIdentifier);
+
+	FString GizmoIdentifier = TEXT("PivotToolGizmoIdentifier");
+	ETransformGizmoSubElements Elements = ETransformGizmoSubElements::TranslateAllAxes;
+	TransformGizmo = GizmoManager->CreateCustomTransformGizmo(Elements, this, UMotionTrailTool::TrailKeyTransformGizmoInstanceIdentifier);
+
+	TransformProxy->OnTransformChanged.AddUObject(this, &UMotionTrailTool::GizmoTransformChanged);
+	TransformProxy->OnBeginTransformEdit.AddUObject(this, &UMotionTrailTool::GizmoTransformStarted);
+	TransformProxy->OnEndTransformEdit.AddUObject(this, &UMotionTrailTool::GizmoTransformEnded);
+	TransformGizmo->SetActiveTarget(TransformProxy, GetToolManager());
+	TransformGizmo->SetVisibility(false);
+	SetupIntegration();
+}
+
+void UMotionTrailTool::SetupIntegration()
+{
+
+	OnSequencersChangedHandle = FLevelEditorSequencerIntegration::Get().GetOnSequencersChanged().AddLambda([this] {
+		for (const TUniquePtr<UE::SequencerAnimTools::FTrailHierarchy>& TrailHierarchy : TrailHierarchies)
+		{
+			TrailHierarchy->Destroy();
+		}
+
+		TrailHierarchies.Reset();
+		SequencerHierarchies.Reset();
+		// TODO: kind of cheap for now, later should check with member TMap<ISequencer*, FTrailHierarchy*> TrackedSequencers
+		for (TWeakPtr<ISequencer> WeakSequencer : FLevelEditorSequencerIntegration::Get().GetSequencers())
+		{
+			TrailHierarchies.Add_GetRef(MakeUnique<UE::SequencerAnimTools::FSequencerTrailHierarchy>(WeakSequencer))->Initialize();
+			SequencerHierarchies.Add(WeakSequencer.Pin().Get(), TrailHierarchies.Last().Get());
+		}
+	});
+
+	for (TWeakPtr<ISequencer> WeakSequencer : FLevelEditorSequencerIntegration::Get().GetSequencers())
+	{
+		TrailHierarchies.Add_GetRef(MakeUnique<UE::SequencerAnimTools::FSequencerTrailHierarchy>(WeakSequencer))->Initialize();
+		SequencerHierarchies.Add(WeakSequencer.Pin().Get(), TrailHierarchies.Last().Get());
+	}
+
+	CommandBindings = MakeShareable(new FUICommandList);
+
+	FMotionTrailCommands::Register();
+
+	const FMotionTrailCommands& Commands = FMotionTrailCommands::Get();
+
+	CommandBindings->MapAction(
+		Commands.TranslateSelectedKeysLeft,
+		FExecuteAction::CreateUObject(this, &UMotionTrailTool::TranslateSelectedKeysLeft),
+		FCanExecuteAction::CreateUObject(this, &UMotionTrailTool::SomeKeysAreSelected));
+
+	CommandBindings->MapAction(
+		Commands.TranslateSelectedKeysRight,
+		FExecuteAction::CreateUObject(this, &UMotionTrailTool::TranslateSelectedKeysRight),
+		FCanExecuteAction::CreateUObject(this, &UMotionTrailTool::SomeKeysAreSelected));
+
+	CommandBindings->MapAction(
+		Commands.FrameSelection,
+		FExecuteAction::CreateUObject(this, &UMotionTrailTool::FrameSelection),
+		FCanExecuteAction::CreateUObject(this, &UMotionTrailTool::SomeKeysAreSelected));
+
+	CommandBindings->MapAction(
+		FGenericCommands::Get().Delete,
+		FExecuteAction::CreateUObject(this, &UMotionTrailTool::DeleteSelectedKeys),
+		FCanExecuteAction::CreateUObject(this, &UMotionTrailTool::SomeKeysAreSelected));
+}
+
+
+void UMotionTrailTool::Shutdown(EToolShutdownType ShutdownType)
+{
+	Super::Shutdown(ShutdownType);
+	GizmoManager->DestroyAllGizmosByOwner(this);
+	ShutdownIntegration();
+}
+
+void UMotionTrailTool::ShutdownIntegration()
+{
+	for (TUniquePtr<UE::SequencerAnimTools::FTrailHierarchy>& TrailHierarchy : TrailHierarchies)
+	{
+		// TODO: just use dtor?
+		TrailHierarchy->Destroy();
+	}
+
+	TrailHierarchies.Reset();
+	FLevelEditorSequencerIntegration::Get().GetOnSequencersChanged().Remove(OnSequencersChangedHandle);
+}
+
+
+
+// Detects Ctrl and Shift key states
+void UMotionTrailTool::OnUpdateModifierState(int ModifierID, bool bIsOn)
+{
+	if (ModifierID == ShiftModifierId)
+	{
+		bShiftModifierOn = bIsOn;
+	}
+	else if (ModifierID == CtrlModifierId)
+	{
+		bCtrlModifierOn = bIsOn;
+	}
+	if (bShiftModifierOn && bCtrlModifierOn)
+	{
+		bAltModifierOn = true;
+		bShiftModifierOn = false;
+		bCtrlModifierOn = false;
+	}
+	//ALT MODIFIER doesn't work, we can not even use FSlateApplication::Get().GetModifierKeys().IsAltDown() 
+	//since if there is an alt it just bails !!!!l
+	
+}
+
+FInputRayHit UMotionTrailTool::IsHitByClick(const FInputDeviceRay& ClickPos)
+{
+	FInputRayHit ReturnHit = FInputRayHit();
+	FViewport* FocusedViewport = GetToolManager()->GetContextQueriesAPI()->GetFocusedViewport();
+	HHitProxy* HitResult = FocusedViewport->GetHitProxy(ClickPos.ScreenPosition.X, ClickPos.ScreenPosition.Y);
+	for (TUniquePtr<UE::SequencerAnimTools::FTrailHierarchy>& TrailHierarchy : TrailHierarchies)
+	{
+		if (TrailHierarchy->IsHitByClick(HitResult))
+		{
+			ReturnHit.bHit = true;
+			break;
+		}
+	}
+
+	return ReturnHit;
+}
+
+void UMotionTrailTool::OnClicked(const FInputDeviceRay& ClickPos)
+{
+	FEditorViewportClient* ViewportClient = ModeManager->GetFocusedViewportClient();
+	FViewport* FocusedViewport = GetToolManager()->GetContextQueriesAPI()->GetFocusedViewport();
+	HHitProxy* HitResult = FocusedViewport->GetHitProxy(ClickPos.ScreenPosition.X, ClickPos.ScreenPosition.Y);
+
+	UE::SequencerAnimTools::FInputClick InputClick(bAltModifierOn, bCtrlModifierOn, bShiftModifierOn);
+	if (RightTarget.bClicked)
+	{
+		InputClick.bIsRightMouse = true;
+	}
+
+	for (TUniquePtr<UE::SequencerAnimTools::FTrailHierarchy>& TrailHierarchy : TrailHierarchies)
+	{
+		TrailHierarchy->HandleClick(ViewportClient, HitResult, InputClick);
+	}
+	UpdateGizmoLocation();
+}
+
+
+void UMotionTrailTool::Render(IToolsContextRenderAPI* RenderAPI)
+{
+	FViewport* FocusedViewport = GetToolManager()->GetContextQueriesAPI()->GetFocusedViewport();
+
+	for (TUniquePtr<UE::SequencerAnimTools::FTrailHierarchy>& TrailHierarchy : TrailHierarchies)
+	{
+		TrailHierarchy->GetRenderer()->Render(RenderAPI->GetSceneView(), FocusedViewport, RenderAPI->GetPrimitiveDrawInterface());
+	}
+}
+
+void UMotionTrailTool::DrawHUD(FCanvas* Canvas, IToolsContextRenderAPI* RenderAPI)
+{
+	FViewport* FocusedViewport = GetToolManager()->GetContextQueriesAPI()->GetFocusedViewport();
+	FEditorViewportClient* ViewportClient = ModeManager->GetFocusedViewportClient();
+	for (TUniquePtr<UE::SequencerAnimTools::FTrailHierarchy>& TrailHierarchy : TrailHierarchies)
+	{
+		TrailHierarchy->GetRenderer()->DrawHUD(ViewportClient, FocusedViewport, RenderAPI->GetSceneView(), Canvas);
+	}
+}
+
+
+void UMotionTrailTool::OnTick(float DeltaTime)
+{
+	for (TUniquePtr<UE::SequencerAnimTools::FTrailHierarchy>& TrailHierarchy : TrailHierarchies)
+	{
+		TrailHierarchy->Update();
+	}
+	UpdateGizmoLocation();
+}
+
+TArray<UObject*> UMotionTrailTool::GetToolProperties(bool bEnabledOnly) const
+{
+	return  TArray<UObject*>();
+}
+
+
+void UMotionTrailTool::GizmoTransformStarted(UTransformProxy* Proxy)
+{
+	//TransactionIndex = GEditor->BeginTransaction(nullptr, LOCTEXT("MoveMotionTrail", "Move Motion Trail"), nullptr);
+
+	StartDragTransform = Proxy->GetTransform();
+
+	bGizmoBeingDragged = true;
+	bManipulatorMadeChange = false;
+	GizmoTransform = StartDragTransform;
+	for (const TUniquePtr<UE::SequencerAnimTools::FTrailHierarchy>& TrailHierarchy : TrailHierarchies)
+	{
+		TrailHierarchy->StartTracking();
+	}
+}
+
+void UMotionTrailTool::GizmoTransformChanged(UTransformProxy* Proxy, FTransform Transform)
+{
+	if (!bGizmoBeingDragged)
+	{
+		return;
+	}
+	GizmoTransform = Transform;
+	FTransform Diff = Transform.GetRelativeTransform(StartDragTransform);
+	for (const TUniquePtr<UE::SequencerAnimTools::FTrailHierarchy>& TrailHierarchy : TrailHierarchies)
+	{
+		FVector WidgetLocation = Transform.GetLocation();
+		FVector Location;
+		bool bIsHandled = TrailHierarchy->ApplyDelta(Diff.GetTranslation(), Diff.GetRotation().Rotator(), WidgetLocation);
+		if (bIsHandled)
+		{
+			bManipulatorMadeChange = true;
+		}
+	}
+	StartDragTransform = Transform;
+	UpdateGizmoLocation();
+
+}
+
+void UMotionTrailTool::GizmoTransformEnded(UTransformProxy* Proxy)
+{
+	for (const TUniquePtr<UE::SequencerAnimTools::FTrailHierarchy>& TrailHierarchy : TrailHierarchies)
+	{
+		TrailHierarchy->EndTracking();
+	}
+	if (bManipulatorMadeChange && TransactionIndex != INDEX_NONE)
+	{
+		//GEditor->EndTransaction();
+	}
+	else if (TransactionIndex != INDEX_NONE)
+	{
+	//	GEditor->CancelTransaction(TransactionIndex);
+	}
+	bGizmoBeingDragged = false;
+	bManipulatorMadeChange = false;
+	UpdateGizmoLocation();
+}
+
+void UMotionTrailTool::UpdateGizmoVisibility()
+{
+	if (TransformGizmo)
+	{
+		bool bAnythingSelected = false;
+
+		for (const TUniquePtr<UE::SequencerAnimTools::FTrailHierarchy>& TrailHierarchy : TrailHierarchies)
+		{
+			FVector Location;
+			bool bIsHandled = TrailHierarchy->IsAnythingSelected();
+			if (bIsHandled)
+			{
+				bAnythingSelected = true;
+			}
+		}
+		TransformGizmo->SetVisibility(bAnythingSelected);
+	}
+}
+
+void UMotionTrailTool::UpdateGizmoLocation()
+{
+	if (!TransformGizmo)
+	{
+		return;
+	}
+	UpdateGizmoVisibility();
+	FVector NewGizmoLocation(0., 0., 0.);
+	int NumSelected = 0;
+	for (const TUniquePtr<UE::SequencerAnimTools::FTrailHierarchy>& TrailHierarchy : TrailHierarchies)
+	{
+		FVector Location;
+		bool bIsHandled = TrailHierarchy->IsAnythingSelected(Location);
+		if (bIsHandled)
+		{
+			NewGizmoLocation += Location;
+			++NumSelected;
+		}
+	}
+	if (NumSelected > 0)
+	{
+		NewGizmoLocation /= (double)NumSelected;
+	}
+	GizmoTransform.SetLocation(NewGizmoLocation);
+	TransformGizmo->ReinitializeGizmoTransform(GizmoTransform);
+}
+
+
+void UMotionTrailTool::TranslateSelectedKeysLeft()
+{
+	if (SomeKeysAreSelected())
+	{
+		const FScopedTransaction Transaction(LOCTEXT("TranslateSelectedKeysLeft", "Translate Selected Keys Left"));
+		for (TUniquePtr<UE::SequencerAnimTools::FTrailHierarchy>& TrailHierarchy : TrailHierarchies)
+		{
+			TrailHierarchy->TranslateSelectedKeys(false);
+		}
+	}
+}
+
+void UMotionTrailTool::TranslateSelectedKeysRight()
+{
+	if (SomeKeysAreSelected())
+	{
+		const FScopedTransaction Transaction(LOCTEXT("TranslateSelectedKeysRight", "Translate Selected Keys Right"));
+		for (TUniquePtr<UE::SequencerAnimTools::FTrailHierarchy>& TrailHierarchy : TrailHierarchies)
+		{
+			TrailHierarchy->TranslateSelectedKeys(true);
+		}
+	}
+}
+
+void UMotionTrailTool::FrameSelection()
+{
+	if (SomeKeysAreSelected())
+	{
+		//if (CurrentViewportClient)
+		{
+			FBox Bounds(EForceInit::ForceInit);
+			TArray<FVector> Positions;
+			for (TUniquePtr<UE::SequencerAnimTools::FTrailHierarchy>& TrailHierarchy : TrailHierarchies)
+			{
+				if (TrailHierarchy->IsAnythingSelected(Positions))
+				{
+					for (const FVector Pos : Positions)
+					{
+						Bounds += Pos;
+						Bounds += Pos + FVector::OneVector * 5.f;
+						Bounds += Pos - FVector::OneVector * 5.f;
+					}
+				}
+			}
+			if (Bounds.IsValid)
+			{
+				//CurrentViewportClient->FocusViewportOnBox(Bounds);
+			}
+		}
+	}
+}
+
+void UMotionTrailTool::SelectNone()
+{
+	if (SomeKeysAreSelected())
+	{
+		for (TUniquePtr<UE::SequencerAnimTools::FTrailHierarchy>& TrailHierarchy : TrailHierarchies)
+		{
+			TrailHierarchy->SelectNone();
+		}
+	}
+}
+
+void UMotionTrailTool::DeleteSelectedKeys()
+{
+	if (SomeKeysAreSelected())
+	{
+		const FScopedTransaction Transaction(LOCTEXT("DeleteSelectedKeys", "Delete Selected Keys"));
+
+		for (TUniquePtr<UE::SequencerAnimTools::FTrailHierarchy>& TrailHierarchy : TrailHierarchies)
+		{
+			TrailHierarchy->DeleteSelectedKeys();
+		}
+	}
+}
+
+bool UMotionTrailTool::SomeKeysAreSelected() const
+{
+	bool bAnythingSelected = false;
+
+	for (const TUniquePtr<UE::SequencerAnimTools::FTrailHierarchy>& TrailHierarchy : TrailHierarchies)
+	{
+		FVector Location;
+		bool bIsHandled = TrailHierarchy->IsAnythingSelected();
+		if (bIsHandled)
+		{
+			bAnythingSelected = true;
+		}
+	}
+	return bAnythingSelected;
+}
+
+
+#undef LOCTEXT_NAMESPACE
diff --git a/Engine/Plugins/Editor/SequencerAnimTools/Source/SequencerAnimTools/Private/Sequencer/AnimationBoneTrail.cpp b/Engine/Plugins/Editor/SequencerAnimTools/Source/SequencerAnimTools/Private/Sequencer/AnimationBoneTrail.cpp
new file mode 100644
--- /dev/null
+++ b/Engine/Plugins/Editor/SequencerAnimTools/Source/SequencerAnimTools/Private/Sequencer/AnimationBoneTrail.cpp
@@ -0,0 +1,231 @@
+// Copyright Epic Games, Inc. All Rights Reserved.
+
+#include "AnimationBoneTrail.h"
+#include "TrailHierarchy.h"
+#include "SequencerTrailHierarchy.h"
+
+#include "Animation/AnimInstance.h"
+#include "Components/SkeletalMeshComponent.h"
+#include "Sequencer/ControlRigLayerInstance.h"
+#include "MovieSceneToolHelpers.h"
+
+#include "ISequencer.h"
+#include "MovieSceneSequence.h"
+#include "Exporters/AnimSeqExportOption.h"
+#include "Animation/AnimSequenceHelpers.h"
+
+namespace UE 
+{
+namespace SequencerAnimTools	
+{
+
+float GetIntervalPerKey(int32 NumFrames, float SequenceLength)
+{
+	return (NumFrames > 1) ? (SequenceLength / (NumFrames - 1)) : MINIMUM_ANIMATION_LENGTH;
+}
+
+void FAnimTrajectoryCache::Evaluate(FTrajectoryCache* ParentTransformCache)
+{
+	TSharedPtr<ISequencer> Sequencer = WeakSequencer.Pin();
+	check(Sequencer);
+
+	if (!SkeletalMeshComponent.IsValid() || !SkeletalMeshComponent->SkeletalMesh || !SkeletalMeshComponent->SkeletalMesh->GetSkeleton())
+	{
+		return;
+	}
+
+	SkeletalMeshComponent->MasterPoseComponent = nullptr;
+
+	CachedAnimSequence->SetSkeleton(SkeletalMeshComponent->SkeletalMesh->GetSkeleton());
+
+	// TODO: for some reason SkeletalMeshComponent becomes invalid sometimes when evaluating, no clue why yet
+	FMovieSceneSequenceTransform MovieSceneSequenceTransform;
+	UAnimSeqExportOption* AnimSeqExportOption = NewObject<UAnimSeqExportOption>(GetTransientPackage(), NAME_None);
+	MovieSceneToolHelpers::ExportToAnimSequence(CachedAnimSequence, AnimSeqExportOption,Sequencer->GetFocusedMovieSceneSequence()->GetMovieScene(), Sequencer.Get(), SkeletalMeshComponent.Get(),
+		Sequencer->GetFocusedTemplateID(), MovieSceneSequenceTransform);
+	AnimSeqExportOption->MarkPendingKill();
+	Sequencer->ForceEvaluate();
+
+	GetSpaceBasedAnimationData(GlobalBoneTransforms);
+	ComponentBoneTransforms.SetNum(GlobalBoneTransforms.Num());
+
+	Spacing = CachedAnimSequence->GetSamplingFrameRate().AsInterval();
+	const double AnimLength = CachedAnimSequence->GetPlayLength();
+	AnimRange = TRange<double>(0.0, AnimLength);
+	for (int32 TrackIndex = 0; TrackIndex < GlobalBoneTransforms.Num(); TrackIndex++)
+	{
+		ComponentBoneTransforms[TrackIndex].SetNum(GlobalBoneTransforms[TrackIndex].Num());
+		double TimeSeconds = 0.0;
+		for (int32 TimeIndex = 0; TimeIndex < GlobalBoneTransforms[TrackIndex].Num(); TimeIndex++, TimeSeconds += Spacing)
+		{
+			ComponentBoneTransforms[TrackIndex][TimeIndex] = GlobalBoneTransforms[TrackIndex][TimeIndex].GetRelativeTransform(ParentTransformCache->GetInterp(TimeSeconds));
+			//ComponentBoneTransforms[TrackIndex][TimeIndex] = GlobalBoneTransforms[TrackIndex][TimeIndex].GetRelativeTransformReverse(ParentTransformCache->GetInterp(TimeSeconds));
+		}
+	}
+
+	bDirty = false;
+}
+
+void FAnimTrajectoryCache::UpdateRange(const TRange<double>& EvalRange, FTrajectoryCache* ParentTransformCache, const int32 BoneIdx)
+{
+	const int32 StartIdx = (EvalRange.GetLowerBoundValue() - AnimRange.GetLowerBoundValue()) / Spacing;
+	const int32 EndIdx = (EvalRange.GetUpperBoundValue() - AnimRange.GetLowerBoundValue()) / Spacing;
+
+	double TimeSeconds = EvalRange.GetLowerBoundValue();
+	for (int32 TimeIndex = StartIdx; TimeIndex < EndIdx; TimeIndex++, TimeSeconds += Spacing)
+	{
+		GlobalBoneTransforms[BoneIdx][TimeIndex] = ComponentBoneTransforms[BoneIdx][TimeIndex] * ParentTransformCache->GetInterp(TimeSeconds);
+	}
+}
+
+void FAnimTrajectoryCache::AddReferencedObjects(FReferenceCollector& Collector)
+{
+	Collector.AddReferencedObject(CachedAnimSequence);
+}
+
+void FAnimTrajectoryCache::GetSpaceBasedAnimationData(TArray<TArray<FTransform>>& OutAnimationDataInComponentSpace)
+{
+	USkeleton* MySkeleton = CachedAnimSequence->GetSkeleton();
+
+	check(MySkeleton);
+	const FReferenceSkeleton& RefSkeleton = MySkeleton->GetReferenceSkeleton();
+
+	int32 NumBones = RefSkeleton.GetNum();
+
+	OutAnimationDataInComponentSpace.Empty(NumBones);
+	OutAnimationDataInComponentSpace.AddZeroed(NumBones);
+
+	// 2d array of animated time [boneindex][time key]
+	const int32 NumKeys = CachedAnimSequence->GetNumberOfSampledKeys();
+	float Interval = UE::SequencerAnimTools::GetIntervalPerKey(NumKeys, CachedAnimSequence->GetPlayLength());
+
+	// allocate arrays
+	for (int32 BoneIndex = 0; BoneIndex < NumBones; ++BoneIndex)
+	{
+		OutAnimationDataInComponentSpace[BoneIndex].AddUninitialized(NumKeys);
+	}
+
+	// If skeleton to track map isn't initialized yet, initialize it
+	if (SkelToTrackIdx.Num() == 0)
+	{
+		SkelToTrackIdx.SetNumUninitialized(NumBones);
+		for (int32 BoneIndex = 0; BoneIndex < NumBones; ++BoneIndex)
+		{
+			FName BoneName = MySkeleton->GetReferenceSkeleton().GetBoneName(BoneIndex);
+			int32 TrackIndex = CachedAnimSequence->GetDataModel()->GetBoneTrackIndexByName(BoneName);
+			SkelToTrackIdx[BoneIndex] = TrackIndex;
+		}
+	}
+
+	// Set component transform for each bone
+	for (int32 BoneIndex = 0; BoneIndex < NumBones; ++BoneIndex)
+	{
+		int32 TrackIndex = SkelToTrackIdx[BoneIndex];
+		int32 ParentBoneIndex = MySkeleton->GetReferenceSkeleton().GetParentIndex(BoneIndex);
+
+		if (TrackIndex != INDEX_NONE)
+		{
+			// fill up keys - calculate PK1 * K1
+			for (int32 Key = 0; Key < NumKeys; ++Key)
+			{
+				FTransform AnimatedLocalKey;
+				UE::Anim::GetBoneTransformFromModel(CachedAnimSequence->GetDataModel(), AnimatedLocalKey, TrackIndex, Interval * Key, CachedAnimSequence->Interpolation);
+
+				if (ParentBoneIndex != INDEX_NONE)
+				{
+					OutAnimationDataInComponentSpace[BoneIndex][Key] = AnimatedLocalKey * OutAnimationDataInComponentSpace[ParentBoneIndex][Key];
+				}
+				else
+				{
+					OutAnimationDataInComponentSpace[BoneIndex][Key] = AnimatedLocalKey;
+				}
+			}
+		}
+		else
+		{
+			// get local spaces from refpose and use that to fill it up
+			FTransform LocalTransform = MySkeleton->GetReferenceSkeleton().GetRefBonePose()[BoneIndex];
+
+			for (int32 Key = 0; Key < NumKeys; ++Key)
+			{
+				if (ParentBoneIndex != INDEX_NONE)
+				{
+					OutAnimationDataInComponentSpace[BoneIndex][Key] = LocalTransform * OutAnimationDataInComponentSpace[ParentBoneIndex][Key];
+				}
+				else
+				{
+					OutAnimationDataInComponentSpace[BoneIndex][Key] = LocalTransform;
+				}
+			}
+		}
+	}
+}
+
+FTransform FAnimBoneTrajectoryCache::GetInterp(const double InTime) const
+{
+	if (AnimTrajectoryCache->GlobalBoneTransforms.Num() == 0)
+	{
+		return Default;
+	}
+
+	const double T = (InTime / AnimTrajectoryCache->Spacing) - FMath::FloorToDouble(InTime / AnimTrajectoryCache->Spacing);
+	const int32 LowIdx = FMath::Clamp(int32((InTime - AnimTrajectoryCache->AnimRange.GetLowerBoundValue()) / AnimTrajectoryCache->Spacing), 0, AnimTrajectoryCache->GlobalBoneTransforms[BoneIdx].Num() - 1);
+	const int32 HighIdx = FMath::Clamp(LowIdx + 1, 0, AnimTrajectoryCache->GlobalBoneTransforms[BoneIdx].Num() - 1);
+
+	FTransform TempBlended;
+	TempBlended.Blend(AnimTrajectoryCache->GlobalBoneTransforms[BoneIdx][LowIdx], AnimTrajectoryCache->GlobalBoneTransforms[BoneIdx][HighIdx], T);
+	return TempBlended;
+}
+
+TArray<double> FAnimBoneTrajectoryCache::GetAllTimesInRange(const TRange<double>& InRange) const
+{
+	if (AnimTrajectoryCache->Spacing == 0.0)
+	{
+		return TArray<double>();
+	}
+
+	TRange<double> GenRange = TRange<double>::Intersection({ AnimTrajectoryCache->AnimRange, InRange });
+
+	TArray<double> AllTimesInRange;
+	AllTimesInRange.Reserve(int(GenRange.Size<double>() / AnimTrajectoryCache->Spacing) + 1);
+	const double FirstTick = FMath::FloorToDouble((GenRange.GetLowerBoundValue()) / AnimTrajectoryCache->Spacing) * AnimTrajectoryCache->Spacing;
+	for (double TickItr = FirstTick; TickItr < GenRange.GetUpperBoundValue(); TickItr += AnimTrajectoryCache->Spacing)
+	{
+		AllTimesInRange.Add(TickItr);
+	}
+
+	return AllTimesInRange;
+}
+
+ETrailCacheState FAnimationBoneTrail::UpdateTrail(const FSceneContext& InSceneContext)
+{
+	if (TrajectoryCache->GetAnimCache()->IsDirty())
+	{
+		return ETrailCacheState::NotUpdated;
+	}
+
+	USkeletalMeshComponent* SkeletalMeshComponent = TrajectoryCache->GetAnimCache()->GetSkeletalMeshComponent();
+	FSequencerTrailHierarchy* SequencerTrailHierarchy = static_cast<FSequencerTrailHierarchy*>(InSceneContext.TrailHierarchy);
+	const FGuid SkelMeshCompGuid = SequencerTrailHierarchy->GetObjectsTracked().FindChecked(SkeletalMeshComponent);
+	FTrajectoryCache* SkelMeshCompTrajectoryCache = SequencerTrailHierarchy->GetAllTrails().FindChecked(SkelMeshCompGuid)->GetTrajectoryTransforms();
+
+
+	if ( !TrajectoryCache->IsValid())
+	{
+		return ETrailCacheState::Dead;
+	}
+
+	if (bForceEvaluateNextTick)
+	{
+		CachedEffectiveRange = TrajectoryCache->GetAnimCache()->GetRange();
+		bForceEvaluateNextTick = false;
+		return ETrailCacheState::Stale;
+	}
+	else
+	{
+		return ETrailCacheState::UpToDate;
+	}
+}
+
+} // namespace MovieScene
+} // namespace UE
diff --git a/Engine/Plugins/Editor/SequencerAnimTools/Source/SequencerAnimTools/Private/Sequencer/AnimationBoneTrail.h b/Engine/Plugins/Editor/SequencerAnimTools/Source/SequencerAnimTools/Private/Sequencer/AnimationBoneTrail.h
new file mode 100644
--- /dev/null
+++ b/Engine/Plugins/Editor/SequencerAnimTools/Source/SequencerAnimTools/Private/Sequencer/AnimationBoneTrail.h
@@ -0,0 +1,124 @@
+// Copyright Epic Games, Inc. All Rights Reserved.
+
+#pragma once
+
+#include "Trail.h"
+#include "TrajectoryCache.h"
+#include "TrajectoryDrawInfo.h"
+
+#include "Sequencer/MovieSceneControlRigParameterSection.h"
+#include "Animation/AnimSequence.h"
+
+namespace UE 
+{
+namespace SequencerAnimTools	
+{
+
+class FAnimTrajectoryCache : public FGCObject
+{
+public:
+	FAnimTrajectoryCache(USkeletalMeshComponent* InSkeletalMeshComponent, TWeakPtr<class ISequencer> InWeakSequencer)
+		: WeakSequencer(InWeakSequencer)
+		, SkeletalMeshComponent(InSkeletalMeshComponent)
+		, CachedAnimSequence(NewObject<UAnimSequence>())
+		, GlobalBoneTransforms()
+		, ComponentBoneTransforms()
+		, SkelToTrackIdx()
+		, AnimRange()
+		, Spacing()
+		, bDirty(true)
+	{
+		CachedAnimSequence->SetSkeleton(InSkeletalMeshComponent->SkeletalMesh->GetSkeleton());
+	}
+
+	void Evaluate(FTrajectoryCache* ParentTrajectoryCache);
+	void UpdateRange(const TRange<double>& EvalRange, FTrajectoryCache* ParentTrajectoryCache, const int32 BoneIdx);
+	const TRange<double>& GetRange() const { return AnimRange; }
+	USkeletalMeshComponent* GetSkeletalMeshComponent() const { return SkeletalMeshComponent.Get(); }
+
+	void MarkAsDirty() { bDirty = true; }
+	bool IsDirty() const { return bDirty; }
+
+	virtual void AddReferencedObjects(FReferenceCollector& Collector) override;
+	virtual FString GetReferencerName() const override
+	{
+		return TEXT("UE::SequencerAnimTools::FAnimTrajectoryCache");
+	}
+
+private:
+
+	void GetSpaceBasedAnimationData(TArray<TArray<FTransform>>& OutAnimationDataInComponentSpace);
+
+	TWeakPtr<ISequencer> WeakSequencer;
+	TWeakObjectPtr<USkeletalMeshComponent> SkeletalMeshComponent;
+	UAnimSequence* CachedAnimSequence;
+	TArray<TArray<FTransform>> GlobalBoneTransforms;
+	TArray<TArray<FTransform>> ComponentBoneTransforms;
+	TArray<int32> SkelToTrackIdx;
+	TRange<double> AnimRange;
+	double Spacing;
+	bool bDirty;
+
+	friend class FAnimBoneTrajectoryCache;
+};
+
+class FAnimBoneTrajectoryCache : public FTrajectoryCache
+{
+public:
+	FAnimBoneTrajectoryCache(const FName& BoneName, TSharedPtr<FAnimTrajectoryCache> InAnimTrajectoryCache)
+		: AnimTrajectoryCache(InAnimTrajectoryCache)
+		, BoneIdx(InAnimTrajectoryCache->CachedAnimSequence->GetSkeleton()->GetReferenceSkeleton().FindBoneIndex(BoneName))
+	{}
+
+	// begin FTrajectoryCache interface
+	virtual const FTransform& Get(const double InTime) const override 
+	{ 
+		return AnimTrajectoryCache->GlobalBoneTransforms[BoneIdx][FMath::Clamp(int32((InTime - AnimTrajectoryCache->AnimRange.GetLowerBoundValue()) / AnimTrajectoryCache->Spacing), 0, AnimTrajectoryCache->ComponentBoneTransforms[0].Num() - 1)];
+	}
+
+	virtual FTransform GetInterp(const double InTime) const override;
+
+	// This cache is read-only
+	virtual void Set(const double InTime, const FTransform& InValue) override {};
+	virtual void SetTransforms(TArray <FTransform>& InTransforms, TArray<FTransform>& InParentTransforms) override {};
+	virtual TArray<double> GetAllTimesInRange(const TRange<double>& InRange) const override;
+	// end FTrajectoryCache interface
+
+	// TODO: true for now
+	bool IsValid() const { return true; }
+	TSharedPtr<FAnimTrajectoryCache> GetAnimCache() const { return AnimTrajectoryCache; }
+	int32 GetBoneIndex() const { return BoneIdx; }
+
+private:
+
+	TSharedPtr<FAnimTrajectoryCache> AnimTrajectoryCache;
+	int32 BoneIdx;
+};
+
+class FAnimationBoneTrail : public FTrail 
+{
+public:
+	FAnimationBoneTrail(USceneComponent* InOwner, const FLinearColor& InColor, const bool bInIsVisible, TSharedPtr<FAnimTrajectoryCache> InAnimTrajectoryCache, const FName& InBoneName, const bool bInIsRootBone)
+		: FTrail(InOwner)
+		, TrajectoryCache(MakeUnique<FAnimBoneTrajectoryCache>(InBoneName, InAnimTrajectoryCache))
+		, CachedEffectiveRange(TRange<double>::Empty())
+		, bIsRootBone(bInIsRootBone)
+	{
+		DrawInfo = MakeUnique<FTrajectoryDrawInfo>(InColor, TrajectoryCache.Get());
+	}
+
+	// FTrail interface
+	virtual ETrailCacheState UpdateTrail(const FSceneContext& InSceneContext) override;
+	virtual FTrajectoryCache* GetTrajectoryTransforms() override { return TrajectoryCache.Get(); }
+	virtual TRange<double> GetEffectiveRange() const override { return CachedEffectiveRange; }
+	// End FTrail interface
+
+private:
+	TUniquePtr<FAnimBoneTrajectoryCache> TrajectoryCache;
+
+	TRange<double> CachedEffectiveRange;
+	bool bIsRootBone;
+};
+
+} // namespace MovieScene
+} // namespace UE
diff --git a/Engine/Plugins/Editor/SequencerAnimTools/Source/SequencerAnimTools/Private/Sequencer/MotionTrailMovieSceneKey.cpp b/Engine/Plugins/Editor/SequencerAnimTools/Source/SequencerAnimTools/Private/Sequencer/MotionTrailMovieSceneKey.cpp
new file mode 100644
--- /dev/null
+++ b/Engine/Plugins/Editor/SequencerAnimTools/Source/SequencerAnimTools/Private/Sequencer/MotionTrailMovieSceneKey.cpp
@@ -0,0 +1,527 @@
+// Copyright Epic Games, Inc. All Rights Reserved.
+
+#include "MotionTrailMovieSceneKey.h"
+#include "MovieSceneTransformTrail.h"
+#include "TrailHierarchy.h"
+#include "Tracks/MovieScene3DTransformTrack.h"
+#include "Sections/MovieScene3DTransformSection.h"
+#include "Channels/MovieSceneFloatChannel.h"
+#include "Channels/MovieSceneChannelProxy.h"
+#include "EntitySystem/Interrogation/MovieSceneInterrogationLinker.h"
+#include "HitProxies.h"
+#include "Trail.h"
+#include "Tools/MotionTrailOptions.h"
+#include "CanvasItem.h"
+#include "ISequencer.h"
+
+namespace UE
+{
+namespace SequencerAnimTools
+{
+
+struct HMotionTrailMovieSceneKeyProxy : public HBaseTrailProxy
+{
+	DECLARE_HIT_PROXY();
+
+	FTransform Transform;
+	FTrailKeyInfo* KeyInfo;
+
+	HMotionTrailMovieSceneKeyProxy(const FGuid& InGuid, const FTransform& InTransform, FTrailKeyInfo* InKeyInfo) :
+		HBaseTrailProxy(InGuid, HPP_UI),
+		Transform(InTransform),
+		KeyInfo(InKeyInfo)
+	{
+	}
+
+};
+IMPLEMENT_HIT_PROXY(HMotionTrailMovieSceneKeyProxy, HBaseTrailProxy);
+
+void FMotionTraiMovieScenelKeyTool::Render(const FGuid& Guid, const FSceneView* View, FViewport* Viewport, FPrimitiveDrawInterface* PDI)
+{
+	const bool bIsVisible = (UMotionTrailToolOptions::GetTrailOptions()->bShowKeys);
+
+	if (!bIsVisible|| PDI == nullptr)
+	{
+		ClearSelection();
+		return;
+	}
+	const bool bHitTesting = PDI && PDI->IsHitTesting();
+	const float KeySize = UMotionTrailToolOptions::GetTrailOptions()->KeySize;
+	for (const TPair<FFrameNumber, TUniquePtr<FTrailKeyInfo>>& FrameKeyPair : Keys.Keys)
+	{
+		if (OwningTrail->GetDrawInfo()->GetCachedViewRange().Contains(OwningTrail->WeakSequencer.Pin()->GetFocusedTickResolution().AsSeconds(FrameKeyPair.Value->FrameNumber)))
+		{
+			FLinearColor KeyColor = CachedSelection.Contains(FrameKeyPair.Value.Get()) ? FLinearColor(1.0,1.0,0.0f) :  UMotionTrailToolOptions::GetTrailOptions()->KeyColor;
+
+			if (bHitTesting)
+			{
+				PDI->SetHitProxy(new HMotionTrailMovieSceneKeyProxy(Guid, FrameKeyPair.Value->Transform, FrameKeyPair.Value.Get()));
+			}
+
+			PDI->DrawPoint(FrameKeyPair.Value->Transform.GetLocation(), KeyColor, KeySize, SDPG_MAX);
+
+			if (bHitTesting)
+			{
+				PDI->SetHitProxy(nullptr);
+			}
+		}
+	}
+}
+
+void FMotionTraiMovieScenelKeyTool::DrawHUD(FEditorViewportClient* ViewportClient, FViewport* Viewport, const FSceneView* View, FCanvas* Canvas)
+{
+	const bool bIsVisible = (UMotionTrailToolOptions::GetTrailOptions()->bShowKeys && UMotionTrailToolOptions::GetTrailOptions()->bShowFrameNumber);
+	if (!bIsVisible)
+	{
+		return;
+	}
+
+	FFrameRate TickResolution = OwningTrail->GetSequencer()->GetFocusedTickResolution();
+	FFrameRate DisplayRate = OwningTrail->GetSequencer()->GetFocusedDisplayRate();
+
+	const float KeySize = UMotionTrailToolOptions::GetTrailOptions()->KeySize;
+	for (const TPair<FFrameNumber, TUniquePtr<FTrailKeyInfo>>& FrameKeyPair : Keys.Keys)
+	{
+		if (OwningTrail->GetDrawInfo()->GetCachedViewRange().Contains(OwningTrail->WeakSequencer.Pin()->GetFocusedTickResolution().AsSeconds(FrameKeyPair.Value->FrameNumber)))
+		{
+			FLinearColor KeyColor = CachedSelection.Contains(FrameKeyPair.Value.Get()) ? FLinearColor(1.0, 1.0, 0.0f) : UMotionTrailToolOptions::GetTrailOptions()->KeyColor;
+
+
+			FVector2D PixelLocation;
+			if (View->WorldToPixel(FrameKeyPair.Value->Transform.GetLocation(), PixelLocation))
+			{
+				PixelLocation /= ViewportClient->GetDPIScale();
+
+				int32 FrameNumber = FFrameRate::TransformTime(FFrameTime(FrameKeyPair.Value->FrameNumber), TickResolution, DisplayRate).RoundToFrame().Value;
+				FCanvasTextItem TextItem(PixelLocation, FText::FromString(FString::FromInt(FrameNumber)), GEngine->GetMediumFont(), KeyColor);
+				TextItem.Scale = FVector2D::UnitVector;
+				//TextItem.EnableShadow(FLinearColor::Black);
+				TextItem.Draw(Canvas);
+
+			}
+		}
+	}
+}
+
+bool FMotionTraiMovieScenelKeyTool::IsSelected(FVector& OutVectorPosition) const
+{
+	if (CachedSelection.Num() > 0)
+	{
+		OutVectorPosition = SelectedKeysTransform.GetLocation();
+		return true;
+	}
+	return false;
+}
+
+bool FMotionTraiMovieScenelKeyTool::IsSelected() const
+{
+	return (CachedSelection.Num() > 0);
+}
+
+bool FMotionTraiMovieScenelKeyTool::HandleClick(const FGuid& Guid, FEditorViewportClient* InViewportClient, HHitProxy* InHitProxy, FInputClick Click)
+{
+
+	if (HMotionTrailMovieSceneKeyProxy* HitProxy = HitProxyCast<HMotionTrailMovieSceneKeyProxy>(InHitProxy))
+	{
+		if (HitProxy->KeyInfo && HitProxy->Guid == Guid)
+		{
+			//	FScopedTransaction ScopedTransaction(LOCTEXT("SelectControlTransaction", "Select Control"), IsInLevelEditor() && !GIsTransacting);
+
+			if (Click.bShiftIsDown)
+			{
+				CachedSelection.Add(HitProxy->KeyInfo);
+			}
+			else if (Click.bCtrlIsDown)
+			{
+				if (CachedSelection.Contains(HitProxy->KeyInfo))
+				{
+					CachedSelection.Remove(HitProxy->KeyInfo);
+				}
+				else
+				{
+					CachedSelection.Add(HitProxy->KeyInfo);
+				}
+			}
+			else
+			{
+				CachedSelection.Reset();
+				CachedSelection.Add(HitProxy->KeyInfo);
+			}
+
+		}
+		else
+		{
+			CachedSelection.Reset();
+			return false;
+		}
+		UpdateSelectedKeysTransform();
+		return true;
+	}
+	CachedSelection.Reset();
+	return false;
+}
+
+
+void FMotionTraiMovieScenelKeyTool::UpdateSelectedKeysTransform()
+{
+	if (CachedSelection.Num() > 0)
+	{
+		FVector NewGizmoLocation = FVector::ZeroVector;
+		for (FTrailKeyInfo* KeyInfo : CachedSelection)
+		{
+			NewGizmoLocation += KeyInfo->Transform.GetLocation();
+		}
+		NewGizmoLocation /= (double)CachedSelection.Num();
+		SelectedKeysTransform.SetLocation(NewGizmoLocation);
+	}
+}
+
+void FMotionTraiMovieScenelKeyTool::OnSectionChanged()
+{
+	if (ShouldRebuildKeys())
+	{
+		ClearSelection();
+		BuildKeys();
+	}
+
+	DirtyKeyTransforms();
+}
+
+void FMotionTraiMovieScenelKeyTool::BuildKeys()
+{
+	TArray<FFrameNumber> KeyTimes = SelectedKeyTimes();
+	Keys.Reset();
+	CachedSelection.Reset();
+
+	UMovieSceneSection* AbsoluteTransformSection = OwningTrail->GetSection();
+	TArrayView<FMovieSceneFloatChannel*> FloatChannels = AbsoluteTransformSection->GetChannelProxy().GetChannels<FMovieSceneFloatChannel>();
+	FloatChannels = FloatChannels.Slice(OwningTrail->GetChannelOffset(), uint8(ETransformChannel::MaxChannel) + 1);
+	for (uint8 ChannelIdx = 0; ChannelIdx <= uint8(ETransformChannel::MaxChannel); ChannelIdx++)
+	{
+		FMovieSceneFloatChannel* FloatChannel = FloatChannels[ChannelIdx];
+		for (int32 Idx = 0; Idx < FloatChannel->GetNumKeys(); Idx++)
+		{
+			FFrameNumber CurTime = FloatChannel->GetTimes()[Idx];
+
+			if (!Keys.Contains(CurTime))
+			{
+				TUniquePtr<FTrailKeyInfo> TempKeyInfo = MakeUnique<FTrailKeyInfo>(CurTime, AbsoluteTransformSection, OwningTrail);
+				Keys.Add(MoveTemp(CurTime), MoveTemp(TempKeyInfo));
+			}
+		}
+	}
+	if (KeyTimes.Num() > 0)
+	{
+		SelectKeyTimes(KeyTimes,false);
+	}
+}
+
+TArray<FKeyHandle> FMotionTraiMovieScenelKeyTool::GetSelectedKeyHandles(FMovieSceneChannel *Channel) 
+{
+	TArray<FKeyHandle> TotalKeyHandles;
+	for (const FTrailKeyInfo* KeyInfo : CachedSelection)
+	{
+		TRange<FFrameNumber> FrameRange(KeyInfo->FrameNumber, KeyInfo->FrameNumber);
+		TArray<FKeyHandle> KeyHandles;
+		TArray<FFrameNumber> KeyTimes;
+		Channel->GetKeys(FrameRange, &KeyTimes, &KeyHandles);
+		for (FKeyHandle& Handle : KeyHandles)
+		{
+			TotalKeyHandles.Add(Handle);
+		}
+	}
+	return TotalKeyHandles;
+}
+
+void FMotionTraiMovieScenelKeyTool::TranslateSelectedKeys(bool bRight)
+{
+	UMovieSceneSection* Section = OwningTrail->GetSection();
+	if (CachedSelection.Num() > 0 && Section && Section->TryModify())
+	{
+		Section->Modify();
+		int32 Shift = bRight ? 1 : -1;
+		FFrameNumber Delta = FQualifiedFrameTime(Shift, OwningTrail->GetSequencer()->GetFocusedDisplayRate()).ConvertTo(OwningTrail->GetSequencer()->GetFocusedTickResolution()).RoundToFrame();
+		TRange<FFrameNumber> SectionNewBounds = Section->GetRange();
+		TArrayView<FMovieSceneFloatChannel*> FloatChannels = Section->GetChannelProxy().GetChannels<FMovieSceneFloatChannel>();
+		FloatChannels = FloatChannels.Slice(OwningTrail->GetChannelOffset(), uint8(ETransformChannel::MaxChannel) + 1);
+		for (uint8 ChannelIdx = 0; ChannelIdx <= uint8(ETransformChannel::MaxChannel); ChannelIdx++)
+		{
+			TArray<FFrameNumber> KeyTimes;
+			FMovieSceneFloatChannel* FloatChannel = FloatChannels[ChannelIdx];
+			TArray<FKeyHandle> KeyHandles = GetSelectedKeyHandles(FloatChannel);
+			if (KeyHandles.Num() > 0)
+			{
+				KeyTimes.SetNum(KeyHandles.Num());
+				FloatChannel->GetKeyTimes(KeyHandles, KeyTimes);
+				FFrameNumber LowestFrameTime = KeyTimes[0];
+				FFrameNumber HighestFrameTime = KeyTimes[0];
+
+				// Perform the transformation
+				for (FFrameNumber& Time : KeyTimes)
+				{
+					FFrameTime KeyTime = Time;
+					Time += Delta;
+					if (Time < LowestFrameTime)
+					{
+						LowestFrameTime = Time;
+					}
+					if (Time > HighestFrameTime)
+					{
+						HighestFrameTime = Time;
+					}
+				}
+				SectionNewBounds = TRange<FFrameNumber>::Hull(SectionNewBounds, TRange<FFrameNumber>(LowestFrameTime, HighestFrameTime + 1));
+				Section->SetRange(SectionNewBounds);
+				FloatChannel->SetKeyTimes(KeyHandles, KeyTimes);
+			}
+
+		}
+
+		for (FTrailKeyInfo* KeyInfo : CachedSelection)
+		{
+			KeyInfo->FrameNumber += Delta;
+		}
+		BuildKeys();
+
+		OwningTrail->GetSequencer()->NotifyMovieSceneDataChanged(EMovieSceneDataChangeType::TrackValueChanged);
+		UpdateSelectedKeysTransform();
+	}
+}
+
+void FMotionTraiMovieScenelKeyTool::DeleteSelectedKeys()
+{
+	UMovieSceneSection* Section = OwningTrail->GetSection();
+	if (CachedSelection.Num() > 0 && Section && Section->TryModify())
+	{
+		Section->Modify();
+		TArrayView<FMovieSceneFloatChannel*> FloatChannels = Section->GetChannelProxy().GetChannels<FMovieSceneFloatChannel>();
+		FloatChannels = FloatChannels.Slice(OwningTrail->GetChannelOffset(), uint8(ETransformChannel::MaxChannel) + 1);
+		for (uint8 ChannelIdx = 0; ChannelIdx <= uint8(ETransformChannel::MaxChannel); ChannelIdx++)
+		{
+			FMovieSceneFloatChannel* FloatChannel = FloatChannels[ChannelIdx];
+			TArray<FKeyHandle> KeyHandles = GetSelectedKeyHandles(FloatChannel);
+			FloatChannel->DeleteKeys(KeyHandles);
+		}
+		OwningTrail->GetSequencer()->NotifyMovieSceneDataChanged(EMovieSceneDataChangeType::TrackValueChanged);
+		CachedSelection.Reset();
+		UpdateSelectedKeysTransform();
+	}
+}
+
+FTrailKeyInfo* FMotionTraiMovieScenelKeyTool::FindKey(const FFrameNumber& FrameNumber) const
+{
+	const TUniquePtr<FTrailKeyInfo>* KeyInfo = Keys.Find(FrameNumber);
+	if (KeyInfo != nullptr)
+	{
+		return (KeyInfo->Get());
+	}
+	//due to seconds to frame issues also try previous and next frames
+	KeyInfo = Keys.Find(FrameNumber + 1);
+	if (KeyInfo != nullptr)
+	{
+		return (KeyInfo->Get());
+	}
+	KeyInfo = Keys.Find(FrameNumber -1);
+	if (KeyInfo != nullptr)
+	{
+		return (KeyInfo->Get());
+	}
+
+	return nullptr;
+}
+
+TArray<FFrameNumber> FMotionTraiMovieScenelKeyTool::SelectedKeyTimes() const
+{
+	TArray<FFrameNumber> Frames;
+	for (const FTrailKeyInfo* KeyInfo : CachedSelection)
+	{
+		Frames.Add(KeyInfo->FrameNumber);
+	}
+	return Frames;
+}
+
+void FMotionTraiMovieScenelKeyTool::SelectKeyTimes(const TArray<FFrameNumber>& Frames, bool KeepSelection)
+{
+	if (KeepSelection == false)
+	{
+		CachedSelection.Reset();
+	}
+	for (const FFrameNumber& FrameNumber : Frames)
+	{
+		FTrailKeyInfo* KeyInfo = FindKey(FrameNumber);
+		if (KeyInfo)
+		{
+			CachedSelection.Add(KeyInfo);
+		}
+	}
+}
+
+TArray<FFrameNumber> FMotionTraiMovieScenelKeyTool::GetTimesFromModifiedTimes(const TArray<FFrameNumber>& ModifiedFrames, const FFrameNumber& LastFrame, const FFrameNumber& Step)
+{
+	TArray<FFrameNumber> TotalModifiedFrames;
+	TSet<int32> Indices;
+	bool bAddLastKey = false; //if we modified the last key we need to go all the way to the end
+	//first we calculate the indices that have changed, we get the two before and the one after
+	for (const FFrameNumber& FrameNumber : ModifiedFrames)
+	{
+		int32 Index = Keys.FindIndex(FrameNumber);
+		if (Index == INDEX_NONE)  //something went wrong pass out empty array which means we will recalc everything
+		{
+			return TotalModifiedFrames;
+		}
+		else
+		{
+			if (Index - 2 >= 0)
+			{
+				if (Indices.Contains(Index - 2) == false)
+				{
+					Indices.Add(Index - 2);
+				}
+			}
+			if (Index - 1 >= 0)
+			{
+				if (Indices.Contains(Index - 1) == false)
+				{
+					Indices.Add(Index - 1);
+				}
+			}
+			if (Indices.Contains(Index) == false)
+			{
+				Indices.Add(Index);
+			}
+			if (Index + 1 < Keys.Num())
+			{
+				if (Indices.Contains(Index + 1) == false)
+				{
+					Indices.Add(Index + 1);
+				}
+				if (Index + 2 == Keys.Num())
+				{
+					bAddLastKey = true;
+				}
+			}
+			else
+			{
+				bAddLastKey = true;
+			}
+		}
+	}
+	for (int32 Index : Indices)
+	{
+		TUniquePtr<FTrailKeyInfo>* FirstKeyInfo =  Keys.FindFromIndex(Index);
+		TUniquePtr<FTrailKeyInfo>* SecondKeyInfo = Keys.FindFromIndex(Index + 1);
+		if (FirstKeyInfo && FirstKeyInfo->IsValid() && SecondKeyInfo && SecondKeyInfo->IsValid())
+		{
+			for (FFrameNumber Frame = FirstKeyInfo->Get()->FrameNumber; Frame < SecondKeyInfo->Get()->FrameNumber; Frame += Step)
+			{
+				TotalModifiedFrames.Add(Frame);
+			}
+			
+		}
+	}
+	if (bAddLastKey)
+	{
+		for (FFrameNumber Frame = TotalModifiedFrames[TotalModifiedFrames.Num() - 1]; Frame <= LastFrame; Frame += Step)
+		{
+			TotalModifiedFrames.Add(Frame);
+		}
+	}
+	return TotalModifiedFrames;
+}
+
+bool FMotionTraiMovieScenelKeyTool::ShouldRebuildKeys()
+{
+	TMap<FFrameNumber, TSet<ETransformChannel>> KeyTimes;
+	TArrayView<FMovieSceneFloatChannel*> FloatChannels = OwningTrail->GetSection()->GetChannelProxy().GetChannels<FMovieSceneFloatChannel>();
+	FloatChannels = FloatChannels.Slice(OwningTrail->GetChannelOffset(), uint8(ETransformChannel::MaxChannel) + 1);
+	for (uint8 ChannelIdx = 0; ChannelIdx <= uint8(ETransformChannel::MaxChannel); ChannelIdx++)
+
+	{
+		FMovieSceneFloatChannel* FloatChannel = FloatChannels[ChannelIdx];
+		for (int32 Idx = 0; Idx < FloatChannel->GetNumKeys(); Idx++)
+		{
+			const FFrameNumber CurTime = FloatChannel->GetTimes()[Idx];
+			KeyTimes.FindOrAdd(CurTime).Add(ETransformChannel(ChannelIdx));
+		}
+	}
+
+	if (KeyTimes.Num() != Keys.Num())
+	{
+		return true;
+	}
+
+	for (const TPair<FFrameNumber, TSet<ETransformChannel>>& TimeKeyPair : KeyTimes)
+	{
+		if (!Keys.Contains(TimeKeyPair.Key))
+		{
+			return true;
+		}
+		for (uint8 ChannelIdx = 0; ChannelIdx <= uint8(ETransformChannel::MaxChannel); ChannelIdx++)
+		{
+			const ETransformChannel TransformChannel = ETransformChannel(ChannelIdx);
+			if ((!TimeKeyPair.Value.Contains(TransformChannel) && Keys[TimeKeyPair.Key]->IdxMap.Contains(TransformChannel)) ||
+				(Keys[TimeKeyPair.Key]->IdxMap.Contains(TransformChannel) && FloatChannels[uint8(TransformChannel)]->GetData().GetIndex(Keys[TimeKeyPair.Key]->IdxMap[TransformChannel]) == INDEX_NONE) ||
+				(TimeKeyPair.Value.Contains(TransformChannel) && !Keys[TimeKeyPair.Key]->IdxMap.Contains(TransformChannel)))
+			{
+				return true;
+			}
+		}
+	}
+	return false;
+}
+
+void FMotionTraiMovieScenelKeyTool::ClearSelection()
+{
+	CachedSelection.Reset();
+}
+
+void FMotionTraiMovieScenelKeyTool::DirtyKeyTransforms()
+{
+	for (const TPair<FFrameNumber, TUniquePtr<FTrailKeyInfo>>& FrameKeyPair : Keys.Keys)
+	{
+		FrameKeyPair.Value->bDirty = true;
+	}
+}
+
+void FMotionTraiMovieScenelKeyTool::UpdateKeysInRange(const TRange<double>& ViewRange)
+{
+	for (const TPair<FFrameNumber, TUniquePtr<FTrailKeyInfo>>& FrameKeyPair : Keys.Keys)
+	{
+		if (FrameKeyPair.Value->bDirty)
+		{
+			FrameKeyPair.Value->UpdateKeyTransform(EGetKeyFrom::FromTrailCache);
+		}
+	}
+	UpdateSelectedKeysTransform();
+}
+
+FTrailKeyInfo::FTrailKeyInfo(const FFrameNumber InFrameNumber, UMovieSceneSection* InSection, FMovieSceneTransformTrail* InOwningTrail):
+  IdxMap()
+, FrameNumber(InFrameNumber)
+, bDirty(true)
+, OwningTrail(InOwningTrail)
+{
+	TArrayView<FMovieSceneFloatChannel*> Channels = InSection->GetChannelProxy().GetChannels<FMovieSceneFloatChannel>();
+	Channels = Channels.Slice(OwningTrail->GetChannelOffset(), uint8(ETransformChannel::MaxChannel) + 1);
+	for (uint8 Idx = 0; Idx <= uint8(ETransformChannel::MaxChannel); Idx++)
+	{
+		const int32 FoundIdx = Channels[Idx]->GetData().FindKey(InFrameNumber);
+		if (FoundIdx != INDEX_NONE)
+		{
+			IdxMap.Add(ETransformChannel(Idx), Channels[Idx]->GetData().GetHandle(FoundIdx));
+		}
+	}
+}
+
+void FTrailKeyInfo::UpdateKeyTransform(EGetKeyFrom UpdateType)
+{	
+	bDirty = false;
+	 if (UpdateType == EGetKeyFrom::FromTrailCache)
+	{
+		const double EvalTime = OwningTrail->GetSequencer()->GetFocusedTickResolution().AsSeconds(FrameNumber);
+		OwningTrail->GetTrajectoryTransforms()->GetInterp(EvalTime,Transform,ParentTransform);
+	}
+}
+
+} // namespace MovieScene
+} // namespace UE
\ No newline at end of file
diff --git a/Engine/Plugins/Editor/SequencerAnimTools/Source/SequencerAnimTools/Private/Sequencer/MotionTrailMovieSceneKey.h b/Engine/Plugins/Editor/SequencerAnimTools/Source/SequencerAnimTools/Private/Sequencer/MotionTrailMovieSceneKey.h
new file mode 100644
--- /dev/null
+++ b/Engine/Plugins/Editor/SequencerAnimTools/Source/SequencerAnimTools/Private/Sequencer/MotionTrailMovieSceneKey.h
@@ -0,0 +1,259 @@
+// Copyright Epic Games, Inc. All Rights Reserved.
+
+#pragma once
+
+#include "CoreMinimal.h"
+#include "UObject/Object.h"
+#include "Misc/FrameNumber.h"
+#include "MovieSceneTracksComponentTypes.h"
+#include "Algo/BinarySearch.h"
+#include "Algo/Sort.h"
+#include "Trail.h"
+
+class UMovieSceneSection;
+struct FMovieSceneChannel;
+
+namespace UE
+{
+namespace SequencerAnimTools
+{
+
+
+class FMovieSceneTransformTrail;
+class FMovieSceneControlRigTransformTrail;
+class FMovieSceneComponentTransformTrail;
+
+enum class ETransformChannel : uint8
+{
+	TranslateX = 0,
+	TranslateY = 1,
+	TranslateZ = 2,
+	RotateX = 3,
+	RotateY = 4,
+	RotateZ = 5,
+	ScaleX = 6,
+	ScaleY = 7,
+	ScaleZ = 8,
+	MaxChannel = 8
+};
+
+enum class EGetKeyFrom
+{
+	FromComponentDelta,
+	FromTrailCache
+};
+
+struct FTrailKeyInfo
+{
+	FTrailKeyInfo(const FFrameNumber InFrameNumber, UMovieSceneSection* InSection,
+		FMovieSceneTransformTrail* InOwningTrail);
+
+	// Re-eval transform or use given one
+	void UpdateKeyTransform(EGetKeyFrom UpdateType);
+
+	// Key Specific info
+	FTransform Transform;
+	FTransform ParentTransform;
+	TMap<ETransformChannel, FKeyHandle> IdxMap;
+	FFrameNumber FrameNumber;
+	bool bDirty;
+
+	FMovieSceneTransformTrail* OwningTrail;
+};
+
+//This class is a limtited version of TSortedMap that also exposes FindIndex and a new FindFromIndex function
+//in order to get direct access to the hidden sorted array.
+class FSortedKeys
+{
+	struct FKeyForward
+	{
+		FORCEINLINE const FFrameNumber& operator()(const TPair<FFrameNumber, TUniquePtr<FTrailKeyInfo>>& Pair) const
+		{
+			return Pair.Key;
+		}
+	};
+public:
+	int32 FindIndex(FFrameNumber Key) const
+	{
+		return Algo::BinarySearchBy(Keys, Key, FKeyForward());
+	}
+
+	const TUniquePtr<FTrailKeyInfo>* Find(FFrameNumber Key) const
+	{
+		int32 Index = FindIndex(Key);
+		if (Index == INDEX_NONE)
+		{
+			return nullptr;
+		}
+		else return &(Keys[Index].Value);
+	}
+
+	int32 FindIndex(FFrameNumber Key) 
+	{
+		return Algo::BinarySearchBy(Keys, Key, FKeyForward());
+	}
+	TUniquePtr<FTrailKeyInfo>* FindFromIndex(int32 Index)
+	{
+		if (Index >= 0 && Index < Keys.Num())
+		{
+			return &(Keys[Index].Value);
+		}
+		return nullptr;
+	}
+	TUniquePtr<FTrailKeyInfo>* Find(FFrameNumber Key) 
+	{
+		int32 Index = FindIndex(Key);
+		if (Index == INDEX_NONE)
+		{
+			return nullptr;
+		}
+		else return &(Keys[Index].Value);
+	}
+	bool Contains(FFrameNumber Key) const
+	{
+		if (Find(Key))
+		{
+			return true;
+		}
+		return false;
+	}
+
+	void Empty(int32 ExpectedNumElements = 0)
+	{
+		Keys.Empty(ExpectedNumElements);
+	}
+
+	void Reset()
+	{
+		Keys.Reset();
+	}
+
+	void Shrink()
+	{
+		Keys.Shrink();
+	}
+
+	void Reserve(int32 Number)
+	{
+		Keys.Reserve(Number);
+	}
+
+	bool IsEmpty() const
+	{
+		return Keys.IsEmpty();
+	}
+
+	int32 Num() const
+	{
+		return Keys.Num();
+	}
+
+	TUniquePtr<FTrailKeyInfo>& FindChecked(FFrameNumber Key)
+	{
+		TUniquePtr<FTrailKeyInfo>* Value = Find(Key);
+		check(Value != nullptr);
+		return *Value;
+	}
+
+	const TUniquePtr<FTrailKeyInfo>& FindChecked(FFrameNumber Key) const
+	{
+		const TUniquePtr<FTrailKeyInfo>* Value = Find(Key);
+		check(Value != nullptr);
+		return *Value;
+	}
+
+	TUniquePtr<FTrailKeyInfo>& operator[](FFrameNumber Key) 
+	{ 
+		return FindChecked(Key);
+	}
+	const TUniquePtr<FTrailKeyInfo>& operator[](FFrameNumber Key) const
+	{ 
+		return FindChecked(Key); 
+	}
+	
+	FORCEINLINE TPair<FFrameNumber, TUniquePtr<FTrailKeyInfo>>* AllocateMemoryForEmplace(FFrameNumber&& InKey)
+	{
+		int32 InsertIndex = Algo::LowerBoundBy(Keys, InKey, FKeyForward());
+		check(InsertIndex >= 0 && InsertIndex <= Keys.Num());
+
+		TPair<FFrameNumber, TUniquePtr<FTrailKeyInfo>>* DataPtr = nullptr;
+		// Adding new one, this may reallocate Pairs
+		Keys.InsertUninitialized(InsertIndex, 1);
+		DataPtr = Keys.GetData() + InsertIndex;
+		return DataPtr;
+	}
+	
+	TUniquePtr<FTrailKeyInfo>& Add(FFrameNumber &&InKey, TUniquePtr<FTrailKeyInfo> && InValue)
+	{
+		TPair<FFrameNumber, TUniquePtr<FTrailKeyInfo>>* DataPtr = AllocateMemoryForEmplace(MoveTemp(InKey));
+
+		new(DataPtr) TPair<FFrameNumber, TUniquePtr<FTrailKeyInfo>>(TPairInitializer<FFrameNumber&&, TUniquePtr<FTrailKeyInfo>&&>(Forward<FFrameNumber>(InKey), Forward< TUniquePtr<FTrailKeyInfo>>(InValue)));
+
+		return DataPtr->Value;
+	}
+	
+	void GetAllKeys(TArray<FTrailKeyInfo*>& OutKeys)
+	{
+		OutKeys.SetNum(Keys.Num());
+		int32 Index = 0;
+		for (TPair < FFrameNumber, TUniquePtr<FTrailKeyInfo>>& Pair : Keys)
+		{
+			OutKeys[Index++] = Pair.Value.Get();
+		}
+	}
+
+	TArray < TPair<FFrameNumber, TUniquePtr<FTrailKeyInfo>> > Keys;
+};
+
+class FMotionTraiMovieScenelKeyTool 
+{
+public:
+	FMotionTraiMovieScenelKeyTool(FMovieSceneTransformTrail* InOwningTrail)
+		: OwningTrail(InOwningTrail), SelectedKeysTransform(FTransform::Identity)
+	{}
+	~FMotionTraiMovieScenelKeyTool() {};
+
+	void Render(const FGuid& Guid, const FSceneView* View, FViewport* Viewport, FPrimitiveDrawInterface* PDI);
+	void DrawHUD(FEditorViewportClient* ViewportClient, FViewport* Viewport, const FSceneView* View, FCanvas* Canvas);
+
+	bool HandleClick(const FGuid& Guid, FEditorViewportClient* InViewportClient, HHitProxy* HitProxy, FInputClick Click);
+	void StartDragging() {};
+	bool IsSelected(FVector& OutVectorPosition) const;
+	bool IsSelected() const;
+	void OnSectionChanged();
+	void UpdateKeysInRange(const TRange<double>& ViewRange);
+	void ClearSelection();
+	FTrailKeyInfo* FindKey(const FFrameNumber& FrameNumber) const;
+	TArray<FFrameNumber> SelectedKeyTimes() const;
+	void SelectKeyTimes(const TArray<FFrameNumber>& Frames, bool KeepSelection);
+	void SelectKeyInfo(FTrailKeyInfo* KeyInfo) { CachedSelection.Add(KeyInfo); }
+	FTransform GetSelectedKeysTransform() const { return SelectedKeysTransform; }
+	void TranslateSelectedKeys(bool bRight);
+	void DeleteSelectedKeys();
+	bool IsAllSelected() const { return Keys.Num() == CachedSelection.Num(); }
+	void GetAllKeys(TArray<FTrailKeyInfo*>& OutKeys) {Keys.GetAllKeys(OutKeys); }
+
+	TArray<FFrameNumber> GetTimesFromModifiedTimes(const TArray<FFrameNumber>& Frames, const FFrameNumber& LastFrame, const FFrameNumber& Step);
+protected:
+
+	void UpdateSelectedKeysTransform();
+	bool ShouldRebuildKeys();
+	void BuildKeys();
+	void DirtyKeyTransforms();
+	TArray<FKeyHandle> GetSelectedKeyHandles(FMovieSceneChannel* Channel);
+
+	//TSortedMap<FFrameNumber, TUniquePtr<FTrailKeyInfo>> Keys;
+	FSortedKeys Keys;
+	TSet<FTrailKeyInfo*> CachedSelection; //current selection
+
+	FMovieSceneTransformTrail* OwningTrail;
+	FTransform SelectedKeysTransform;
+
+	friend class FMovieSceneTransformTrail;
+	friend class FMovieSceneComponentTransformTrail;
+	friend class FMovieSceneControlRigTransformTrail;
+};
+
+
+} // namespace MovieScene
+} // namespace UE
\ No newline at end of file
diff --git a/Engine/Plugins/Editor/SequencerAnimTools/Source/SequencerAnimTools/Private/Sequencer/MovieSceneTransformTrail.cpp b/Engine/Plugins/Editor/SequencerAnimTools/Source/SequencerAnimTools/Private/Sequencer/MovieSceneTransformTrail.cpp
new file mode 100644
--- /dev/null
+++ b/Engine/Plugins/Editor/SequencerAnimTools/Source/SequencerAnimTools/Private/Sequencer/MovieSceneTransformTrail.cpp
@@ -0,0 +1,785 @@
+// Copyright Epic Games, Inc. All Rights Reserved.
+
+#include "MovieSceneTransformTrail.h"
+#include "TrailHierarchy.h"
+
+#include "ISequencer.h"
+#include "Tracks/MovieScene3DTransformTrack.h"
+#include "Sections/MovieScene3DTransformSection.h"
+#include "Channels/MovieSceneFloatChannel.h"
+#include "Channels/MovieSceneChannelProxy.h"
+#include "Systems/MovieSceneComponentTransformSystem.h"
+#include "MovieSceneSequence.h"
+#include "MovieSceneTracksComponentTypes.h"
+#include "Sequencer/MovieSceneControlRigParameterSection.h"
+
+#include "BaseGizmos/GizmoComponents.h"
+#include "BaseGizmos/TransformGizmo.h"
+
+#include "ViewportWorldInteraction.h"
+#include "GameFramework/Actor.h"
+#include "Components/SceneComponent.h"
+#include "MovieSceneToolHelpers.h"
+#include "ActorForWorldTransforms.h"
+#include "IControlRigObjectBinding.h"
+#include "LevelSequence.h"
+#include "LevelSequenceActor.h"
+#include "LevelSequencePlayer.h"
+#include "Evaluation/MovieScenePlayback.h"
+#include "Tools/ControlRigSnapSettings.h"
+#include "ScopedTransaction.h"
+
+#define LOCTEXT_NAMESPACE "SequencerAnimTools"
+
+namespace UE
+{
+namespace SequencerAnimTools
+{
+
+/****************************************************************************************************************
+*
+*   Base Movie Scene Transform Trail
+*
+*****************************************************************************************************************/
+
+
+FMovieSceneTransformTrail::FMovieSceneTransformTrail(USceneComponent* InOwner, const FLinearColor& InColor, const bool bInIsVisible, TWeakObjectPtr<UMovieSceneTrack> InWeakTrack, TSharedPtr<class ISequencer> InSequencer, ULevelSequencePlayer* InSequencePlayer)
+	: FTrail(InOwner)
+	, WeakSequencer(InSequencer)
+	, LevelSequencePlayer(InSequencePlayer)
+	, CachedEffectiveRange(TRange<double>::Empty())
+	, bIsSelected(false)
+	, TrajectoryCache()
+	, LastTransformSectionSig(FGuid())
+	, CachedHierarchyGuid()
+	, WeakTrack(InWeakTrack)
+{
+	TrajectoryCache = MakeUnique<FArrayTrajectoryCache>(0.01, GetEffectiveSectionRange(0));
+	KeyTool = MakeUnique< FMotionTraiMovieScenelKeyTool>(this);
+	DrawInfo = MakeUnique<FTrajectoryDrawInfo>(InColor, TrajectoryCache.Get());
+}
+
+UMovieSceneSection* FMovieSceneTransformTrail::GetSection() const
+{
+	if (WeakTrack.IsValid())
+	{
+		UMovieSceneSection* Section = WeakTrack.Get()->GetSectionToKey();
+		if (Section)
+		{
+			return Section;
+		}
+		else if( WeakTrack.Get()->GetAllSections().Num() >0)
+		{
+			return WeakTrack.Get()->GetAllSections()[0];
+		}
+	}
+	return  nullptr;
+}
+
+bool FMovieSceneTransformTrail::GetEditedTimes(const FTrailHierarchy* TrailHierarchy,const FFrameNumber& LastFrame, TArray<FFrameNumber>& OutEditedTimes)
+{
+	if (EditedTimes.IsSet())
+	{
+		CalculateEditedTimes(TrailHierarchy,LastFrame);
+		if (EditedTimes.IsSet())
+		{
+			OutEditedTimes = EditedTimes.GetValue();
+			return true;
+		}
+	}
+	return false;
+}
+
+void FMovieSceneTransformTrail::CalculateEditedTimes(const FTrailHierarchy* TrailHierarchy, const FFrameNumber& LastFrame)
+{
+	if (bIsSelected || KeyTool->IsAllSelected())
+	{
+		EditedTimes.Reset();
+	}
+	else
+	{
+		const FFrameNumber Step = TrailHierarchy->GetFramesPerSegment();
+		TArray<FFrameNumber> SelectedTimes  = KeyTool->SelectedKeyTimes();
+		EditedTimes = KeyTool->GetTimesFromModifiedTimes(SelectedTimes,LastFrame,Step);
+		if (EditedTimes.GetValue().Num() == 0)
+		{
+			EditedTimes.Reset();
+		}
+	}
+}
+void FMovieSceneTransformTrail::UpdateKeysInRange(const TRange<double>& ViewRange)
+{
+	KeyTool->UpdateKeysInRange(ViewRange);
+}
+
+ETrailCacheState FMovieSceneTransformTrail::UpdateTrail(const FSceneContext& InSceneContext)
+{
+	CachedHierarchyGuid = InSceneContext.YourNode;
+	UMovieSceneSection* Section = GetSection();
+	TSharedPtr<ISequencer> Sequencer = WeakSequencer.Pin();
+
+	FGuid SequencerBinding = FGuid();
+	if (Sequencer && Section)
+	{ // TODO: expensive, but for some reason Section stays alive even after it is deleted
+		Sequencer->GetFocusedMovieSceneSequence()->GetMovieScene()->FindTrackBinding(*Cast<UMovieSceneTrack>(Section->GetOuter()), SequencerBinding);
+	}
+
+	if (!Sequencer || !Section || !SequencerBinding.IsValid())
+	{
+		return ETrailCacheState::Dead;
+	}
+	
+	const bool bTrackUnchanged = Section->GetSignature() == LastTransformSectionSig;
+
+	if (!bTrackUnchanged  || bForceEvaluateNextTick)
+	{
+
+		KeyTool->OnSectionChanged();
+		TrajectoryCache->UpdateCacheTimes(InSceneContext.EvalTimes);
+
+		CacheState = ETrailCacheState::Stale;
+		bForceEvaluateNextTick = false;
+		LastTransformSectionSig = Section->GetSignature();
+
+		UpdateCacheTimes(InSceneContext);
+				
+	}
+	else 
+	{
+		TrajectoryCache->UpdateCacheTimes(InSceneContext.EvalTimes);
+		KeyTool->UpdateSelectedKeysTransform();
+		CacheState = ETrailCacheState::UpToDate;
+	}
+	return CacheState;
+}
+
+void FMovieSceneTransformTrail::DrawHUD(FEditorViewportClient* ViewportClient, FViewport* Viewport, const FSceneView* View, FCanvas* Canvas)
+{
+	KeyTool->DrawHUD(ViewportClient, Viewport, View, Canvas);
+}
+
+void FMovieSceneTransformTrail::Render(const FGuid& Guid, const FSceneView* View, FViewport* Viewport, FPrimitiveDrawInterface* PDI)
+{
+	KeyTool->Render(Guid, View, Viewport, PDI);
+}
+
+bool FMovieSceneTransformTrail::HandleAltClick(FEditorViewportClient* InViewportClient, HMotionTrailProxy* Proxy, FInputClick Click)
+{
+	return false;
+}
+
+bool FMovieSceneTransformTrail::HandleClick(const FGuid& Guid, FEditorViewportClient* InViewportClient, HHitProxy* InHitProxy, FInputClick Click)
+{
+	bool bKeySelected = KeyTool->HandleClick(Guid,InViewportClient, InHitProxy, Click);
+	if (bKeySelected)
+	{
+		if (Click.bShiftIsDown == false || Click.bCtrlIsDown == false)
+		{
+			bIsSelected = false;
+		}
+		return true;
+	}
+
+	if (HMotionTrailProxy* HitProxy = HitProxyCast<HMotionTrailProxy>(InHitProxy))
+	{
+		if (HitProxy->Guid == Guid)
+		{
+
+			//	FScopedTransaction ScopedTransaction(LOCTEXT("SelectControlTransaction", "Select Control"), IsInLevelEditor() && !GIsTransacting);
+			if (Click.bAltIsDown)
+			{
+				return HandleAltClick(InViewportClient, HitProxy, Click);
+			}
+			if (Click.bShiftIsDown)
+			{
+				bIsSelected = true;
+				SelectedPos = HitProxy->Point;
+			}
+			else if (Click.bCtrlIsDown)
+			{
+				if (bIsSelected)
+				{
+					bIsSelected = false;
+				}
+				else
+				{
+					bIsSelected = true;
+					SelectedPos = HitProxy->Point;
+				}
+			}
+			else
+			{
+				KeyTool->ClearSelection();
+				bIsSelected = true;
+				SelectedPos = HitProxy->Point;
+			}
+			return true;
+		}
+	}
+
+	bIsSelected = false;
+	return false;
+}
+
+bool FMovieSceneTransformTrail::IsAnythingSelected(FVector& OutVectorPosition) const
+{
+	if (bIsSelected)
+	{
+		OutVectorPosition = SelectedPos;
+		return true;
+	}
+	return KeyTool->IsSelected(OutVectorPosition);
+}
+
+void FMovieSceneTransformTrail::TranslateSelectedKeys(bool bRight)
+{
+	KeyTool->TranslateSelectedKeys(bRight);
+}
+
+void FMovieSceneTransformTrail::DeleteSelectedKeys()
+{
+	KeyTool->DeleteSelectedKeys();
+}
+
+void FMovieSceneTransformTrail::SelectNone()
+{
+	KeyTool->ClearSelection();
+}
+
+bool FMovieSceneTransformTrail::IsAnythingSelected() const
+{
+	return (bIsSelected || KeyTool->IsSelected());
+}
+
+bool FMovieSceneTransformTrail::IsTrailSelected() const
+{
+	return bIsSelected;
+}
+
+bool FMovieSceneTransformTrail::StartTracking()
+{
+	return false;
+}
+
+bool FMovieSceneTransformTrail::ApplyDelta(const FVector& Pos, const FRotator& Rot, const FVector& WidgetLocation)
+{
+	return false;;
+}
+
+bool FMovieSceneTransformTrail::EndTracking()
+{
+	return false;
+}
+
+TRange<double> FMovieSceneTransformTrail::GetEffectiveSectionRange(int32 ChannelOffset) const
+{
+	UMovieSceneSection* TransformSection = GetSection();
+	if (!TransformSection)
+	{
+		return TRange<double>();
+	}
+	TSharedPtr<ISequencer> Sequencer = WeakSequencer.Pin();
+	check(TransformSection && Sequencer);
+
+	TRange<FFrameNumber> EffectiveRange = TRange<FFrameNumber>::Empty();
+
+	TArrayView<FMovieSceneFloatChannel*> Channels = TransformSection->GetChannelProxy().GetChannels<FMovieSceneFloatChannel>();
+	for (int32 ChannelIdx = ChannelOffset; ChannelIdx <= ChannelOffset + uint8(ETransformChannel::MaxChannel); ChannelIdx++)
+	{
+		FMovieSceneFloatChannel* Channel = Channels[ChannelIdx];
+		EffectiveRange = TRange<FFrameNumber>::Hull(EffectiveRange, Channel->ComputeEffectiveRange());
+	}
+
+	EffectiveRange = TRange<FFrameNumber>::Intersection(EffectiveRange, TransformSection->GetRange());
+
+	TRange<double> SectionRangeSeconds = TRange<double>(
+		Sequencer->GetFocusedTickResolution().AsSeconds(EffectiveRange.GetLowerBoundValue()),
+		Sequencer->GetFocusedTickResolution().AsSeconds(EffectiveRange.GetUpperBoundValue())
+	);
+		
+	// TODO: clip by movie scene range? try movie keys outside of movie scene range
+
+	return SectionRangeSeconds;
+}
+
+/****************************************************************************************************************
+*
+*	Movie Scene Transform Component Trail
+*
+*****************************************************************************************************************/
+
+
+void FMovieSceneComponentTransformTrail::UpdateCacheTimes(const FSceneContext& InSceneContext)
+{
+	TSharedPtr<ISequencer> Sequencer = WeakSequencer.Pin();
+	check(Sequencer);
+	FFrameRate TickResolution = Sequencer->GetFocusedTickResolution();
+	const FTrailEvaluateTimes& EvaluateTimes = InSceneContext.EvalTimes;
+	TArray<FFrameNumber> Frames;
+	bool bUseEditedTimes = false;
+	if (EditedTimes.IsSet())
+	{
+		const FFrameTime LastFrame = EvaluateTimes.EvalTimes[EvaluateTimes.EvalTimes.Num() - 1] * TickResolution;
+
+		CalculateEditedTimes(InSceneContext.TrailHierarchy,LastFrame.RoundToFrame());
+		Frames = EditedTimes.GetValue();
+		bUseEditedTimes = true;
+	}
+	else
+	{
+		//Todo mz make time in frames
+		int32 Index = 0;
+		Frames.SetNum(EvaluateTimes.EvalTimes.Num());
+		for (const double Time : EvaluateTimes.EvalTimes)
+		{
+			const FFrameTime TickTime = Time * TickResolution;
+			Frames[Index++] = TickTime.RoundToFrame();
+		}
+	}
+
+	TArray<FTransform> OutWorldTransforms;
+	FActorForWorldTransforms Actors;
+	Actors.Actor = Component->GetTypedOuter<AActor>();
+	Actors.Component = Component;
+	MovieSceneToolHelpers::GetActorWorldTransforms(GetSequencer().Get(), Actors, Frames, OutWorldTransforms);
+	TArray<FTransform> ParentCache;
+	if (EditedTimes.IsSet())
+	{
+		for (int32 Index = 0; Index < Frames.Num(); ++Index)
+		{
+			const FFrameNumber& FrameNumber = Frames[Index];
+			double Sec = TickResolution.AsSeconds(FFrameTime(FrameNumber));
+			GetTrajectoryTransforms()->Set(Sec, OutWorldTransforms[Index]);
+		}
+	}
+	else
+	{
+		GetTrajectoryTransforms()->SetTransforms(OutWorldTransforms, ParentCache);
+	}
+
+	KeyTool->UpdateKeysInRange(InSceneContext.TrailHierarchy->GetViewRange());
+	EditedTimes.Reset();
+}
+
+bool FMovieSceneComponentTransformTrail::StartTracking()
+{
+	UMovieSceneSection* Section = GetSection();
+	if (!Section)
+	{
+		return false;
+	}
+
+	if (IsAnythingSelected())
+	{
+		EditedTimes.Reset();
+		Section->Modify();
+		KeyTool->StartDragging();
+		return true;
+	}
+	return false;
+}
+
+static bool SetActorTransform(ISequencer* Sequencer, AActor *Actor, UMovieScene3DTransformSection* TransformSection, const TArray<FFrameNumber>& Frames,
+	const TArray<FTransform>& WorldTransformsToSnapTo)
+{
+	if (!Sequencer || !Sequencer->GetFocusedMovieSceneSequence())
+	{
+		return false;
+	}
+	if (!TransformSection)
+	{
+		return false;
+	}
+
+	TransformSection->Modify();
+
+	TArray<FTransform> ParentWorldTransforms;
+	AActor* ParentActor = Actor->GetAttachParentActor();
+	if (ParentActor)
+	{
+		FActorForWorldTransforms ActorSelection;
+		ActorSelection.Actor = ParentActor;
+		ActorSelection.SocketName = Actor->GetAttachParentSocketName();
+		MovieSceneToolHelpers::GetActorWorldTransforms(Sequencer, ActorSelection, Frames, ParentWorldTransforms);
+	}
+	else
+	{
+		ParentWorldTransforms.SetNumUninitialized(Frames.Num());
+		for (FTransform& Transform : ParentWorldTransforms)
+		{
+			Transform = FTransform::Identity;
+		}
+	}
+
+	TArrayView<FMovieSceneFloatChannel*> Channels = TransformSection->GetChannelProxy().GetChannels<FMovieSceneFloatChannel>();
+	for (int32 Index = 0; Index < Frames.Num(); ++Index)
+	{
+		const FFrameNumber& Frame = Frames[Index];
+		FTransform ParentTransform = ParentWorldTransforms[Index];
+		FTransform WorldTransform = WorldTransformsToSnapTo[Index];
+		FTransform LocalTransform = WorldTransform.GetRelativeTransform(ParentTransform);
+		FVector Location = LocalTransform.GetLocation();
+		FRotator Rotation = LocalTransform.GetRotation().Rotator();
+			
+		TMovieSceneChannelData<FMovieSceneFloatValue> ChannelData = Channels[0]->GetData();
+		MovieSceneToolHelpers::SetOrAddKey(ChannelData, Frame, Location.X);
+		ChannelData = Channels[1]->GetData();
+		MovieSceneToolHelpers::SetOrAddKey(ChannelData, Frame, Location.Y);
+		ChannelData = Channels[2]->GetData();
+		MovieSceneToolHelpers::SetOrAddKey(ChannelData, Frame, Location.Z);			
+	}
+
+	Channels[0]->AutoSetTangents();
+	Channels[1]->AutoSetTangents();
+	Channels[2]->AutoSetTangents();
+		
+	Sequencer->NotifyMovieSceneDataChanged(EMovieSceneDataChangeType::MovieSceneStructureItemsChanged);
+	return true;
+}
+
+bool FMovieSceneComponentTransformTrail::ApplyDelta(const FVector& Pos, const FRotator& Rot, const FVector& WidgetLocation)
+{
+	if (Component.IsValid() == false)
+	{
+		return false;
+	}
+	UMovieSceneSection* Section = GetSection();
+	if (!Section)
+	{
+		return false;
+	}
+	UMovieScene3DTransformSection* TransformSection = Cast<UMovieScene3DTransformSection>(Section);
+	if (!Section)
+	{
+		return false;
+	}
+
+	if (Pos.IsNearlyZero() && Rot.IsNearlyZero())
+	{
+		return false;
+	}
+
+	if (IsAnythingSelected())
+	{
+		TArray<FTransform> Transforms;
+		TArray<FFrameNumber> Times;
+		if (IsTrailSelected())
+		{
+			TArray<FTrailKeyInfo*> Keys;
+			KeyTool->GetAllKeys(Keys);
+			for (FTrailKeyInfo* KeyInfo : Keys)
+			{
+				FTransform NewTransform(KeyInfo->Transform);
+				NewTransform.SetLocation(NewTransform.GetLocation() + Pos);
+				KeyInfo->Transform = NewTransform;
+				Transforms.Add(NewTransform);
+				Times.Add(KeyInfo->FrameNumber);
+			}
+		}
+		else
+		{
+			for (FTrailKeyInfo* KeyInfo : KeyTool->CachedSelection)
+			{
+				FTransform NewTransform(KeyInfo->Transform);
+				NewTransform.SetLocation(NewTransform.GetLocation() + Pos);
+				KeyInfo->Transform = NewTransform;
+				Transforms.Add(NewTransform);
+				Times.Add(KeyInfo->FrameNumber);
+			}
+			TArray<FFrameNumber> EditTimesSet;
+			EditedTimes = EditTimesSet;
+		}
+		AActor* Actor = Component.Get()->GetTypedOuter<AActor>();
+		SetActorTransform(GetSequencer().Get(), Actor, TransformSection, Times, Transforms);
+		KeyTool->UpdateSelectedKeysTransform();
+		return true;
+	}
+	return false;
+}
+
+bool FMovieSceneComponentTransformTrail::EndTracking()
+{
+	EditedTimes.Reset();
+	UMovieSceneSection* Section = GetSection();
+	if (!Section)
+	{
+		return false;
+	}
+
+	if (IsAnythingSelected())
+	{
+		//need to broadcast this so we can update other systems(including our own).
+		AActor* Actor = Component.Get()->GetTypedOuter<AActor>();
+		GEngine->BroadcastOnActorMoved(Actor);
+		return true;
+	}
+	return false;
+}
+
+bool FMovieSceneComponentTransformTrail::HandleAltClick( FEditorViewportClient* InViewportClient, HMotionTrailProxy* Proxy,FInputClick Click)
+{
+	if (Component.IsValid() == false)
+	{
+		return false;
+	}
+	UMovieSceneSection* Section = GetSection();
+	if (!Section)
+	{
+		return false;
+	}
+	UMovieScene3DTransformSection* TransformSection = Cast<UMovieScene3DTransformSection>(Section);
+	if (!Section)
+	{
+		return false;
+	}
+
+	const FScopedTransaction Transaction(LOCTEXT("InsertKey", "Insert Key"));
+
+	TArray<FFrameNumber> Times;
+	TArray<FTransform> Transforms;
+	FFrameRate TickResolution = GetSequencer()->GetFocusedTickResolution();
+	FFrameRate DisplayResolution = GetSequencer()->GetFocusedDisplayRate();
+	//Snap to Frame. Assuming we want keys on frames.
+	FFrameTime GlobalTime = TickResolution.AsFrameTime(Proxy->Second);
+	FFrameTime DisplayTime = FFrameRate::TransformTime(GlobalTime, DisplayResolution, TickResolution);
+	DisplayTime = DisplayTime.RoundToFrame();
+	FFrameTime EvalTime = FFrameRate::TransformTime(DisplayTime, TickResolution, DisplayResolution);
+	Times.Add(EvalTime.RoundToFrame());
+	const FTransform Transform = GetTrajectoryTransforms()->GetInterp(Proxy->Second);
+	Transforms.Add(Transform);
+	AActor* Actor = Component.Get()->GetTypedOuter<AActor>();
+	SetActorTransform(GetSequencer().Get(), Actor, TransformSection, Times, Transforms);
+	KeyTool->UpdateSelectedKeysTransform();
+	return true;
+}
+
+/****************************************************************************************************************
+*
+*   Control Rig Transform Trail
+*
+*****************************************************************************************************************/
+
+FMovieSceneControlRigTransformTrail::FMovieSceneControlRigTransformTrail(USceneComponent* SceneComponent, const FLinearColor& InColor, const bool bInIsVisible, TWeakObjectPtr<UMovieSceneTrack> InWeakTrack, TSharedPtr<ISequencer> InSequencer, ULevelSequencePlayer* InPlayer, const FName& InControlName)
+: FMovieSceneTransformTrail(SceneComponent, InColor, bInIsVisible, InWeakTrack, InSequencer, InPlayer)
+, ControlName(InControlName)
+{
+}
+
+FTransform FMovieSceneControlRigTransformTrail::EvaluateChannelsAtTime(TArrayView<FMovieSceneFloatChannel*> Channels, FFrameTime Time) const
+{
+	FVector3f TempTranslation;
+	Channels[0]->Evaluate(Time, TempTranslation.X);
+	Channels[1]->Evaluate(Time, TempTranslation.Y);
+	Channels[2]->Evaluate(Time, TempTranslation.Z);
+	FRotator TempRotation;
+	Channels[3]->Evaluate(Time, TempRotation.Roll);
+	Channels[4]->Evaluate(Time, TempRotation.Pitch);
+	Channels[5]->Evaluate(Time, TempRotation.Yaw);
+	FVector3f TempScale3D;
+	Channels[6]->Evaluate(Time, TempScale3D.X);
+	Channels[7]->Evaluate(Time, TempScale3D.Y);
+	Channels[8]->Evaluate(Time, TempScale3D.Z);
+	FTransform TempTransform = FTransform(TempRotation, TempTranslation, TempScale3D);
+	TempTransform.NormalizeRotation();
+	return TempTransform;
+}
+
+void FMovieSceneControlRigTransformTrail::UpdateCacheTimes(const FSceneContext& InSceneContext)
+{
+	/*
+	We don't do this anymore since we instead update the cache for EVERY drawn control rig trail otherwise
+	we evaluate the rig again for each trail.
+	*/
+}
+
+bool FMovieSceneControlRigTransformTrail::StartTracking()
+{
+	UMovieSceneControlRigParameterSection* Section = Cast<UMovieSceneControlRigParameterSection>(GetSection());
+	if (!Section)
+	{
+		return false;
+	}
+	UControlRig* ControlRig = Section->GetControlRig();
+	if (!ControlRig)
+	{
+		return false ;
+	}
+	if(IsAnythingSelected())
+	{ 
+		EditedTimes.Reset();
+		Section->Modify();
+		ControlRig->Modify();
+		KeyTool->StartDragging();
+		return true;
+	}
+	return false;
+}
+
+int32 FMovieSceneControlRigTransformTrail::GetChannelOffset() const
+{
+	UMovieSceneControlRigParameterSection* CRParamSection = Cast<UMovieSceneControlRigParameterSection>(GetSection());
+	if (CRParamSection)
+	{
+		FChannelMapInfo* pChannelIndex = CRParamSection->ControlChannelMap.Find(ControlName);
+		return pChannelIndex->ChannelIndex;;
+	}
+	return 0;
+}
+
+bool FMovieSceneControlRigTransformTrail::ApplyDelta(const FVector& Pos, const FRotator& Rot, const FVector& WidgetLocation)
+{
+	UMovieSceneControlRigParameterSection* Section = Cast<UMovieSceneControlRigParameterSection>(GetSection());
+	if (!Section)
+	{
+		return false;
+	}
+	UControlRig* ControlRig = Section->GetControlRig();
+	if (!ControlRig)
+	{
+		return false;
+	}
+
+	if (Pos.IsNearlyZero() && Rot.IsNearlyZero())
+	{
+		return false;
+	}
+
+	if (IsAnythingSelected())
+	{
+		FFrameRate TickResolution = GetSequencer()->GetFocusedTickResolution();
+		FRigControlModifiedContext Context;
+		Context.SetKey = EControlRigSetKey::Always;
+		IMovieScenePlayer* Player = GetSequencer().Get();
+
+		auto EvalControlRig = [&Context, Pos, TickResolution, ControlRig,Player, this](FTrailKeyInfo* KeyInfo)
+		{
+			if (KeyInfo)
+			{
+				Context.LocalTime = TickResolution.AsSeconds(FFrameTime(KeyInfo->FrameNumber));
+				Context.KeyMask = (uint32)EControlRigContextChannelToKey::Translation;
+				FFrameTime GlobalTime(KeyInfo->FrameNumber);
+				FMovieSceneContext MovieSceneContext = FMovieSceneContext(FMovieSceneEvaluationRange(GlobalTime, TickResolution), Player->GetPlaybackStatus()).SetHasJumped(true);
+
+				FTransform NewTransform(KeyInfo->Transform);
+				NewTransform.SetLocation(NewTransform.GetLocation() + Pos);
+				KeyInfo->Transform = NewTransform;
+				NewTransform = NewTransform.GetRelativeTransform(KeyInfo->ParentTransform);
+				GetSequencer()->GetEvaluationTemplate().Evaluate(MovieSceneContext, *Player);
+				ControlRig->Evaluate_AnyThread();
+				ControlRig->SetControlGlobalTransform(ControlName, NewTransform, true, Context, false);
+			}
+		};
+
+		if (IsTrailSelected())
+		{
+			TArray<FTrailKeyInfo*> Keys;
+			KeyTool->GetAllKeys(Keys);
+			for (FTrailKeyInfo* KeyInfo : Keys)
+			{
+				EvalControlRig(KeyInfo);
+			}
+			SelectedPos += Pos;
+		}
+		else
+		{
+			for (FTrailKeyInfo* KeyInfo : KeyTool->CachedSelection)
+			{
+				EvalControlRig(KeyInfo);
+			}
+			TArray<FFrameNumber> EditTimesSet;
+			EditedTimes = EditTimesSet;
+		}
+		KeyTool->UpdateSelectedKeysTransform();
+		return true;
+	}
+	return false;
+}
+
+bool FMovieSceneControlRigTransformTrail::EndTracking()
+{
+	UMovieSceneControlRigParameterSection* Section = Cast<UMovieSceneControlRigParameterSection>(GetSection());
+	if (!Section)
+	{
+		return false;
+	}
+	UControlRig* ControlRig = Section->GetControlRig();
+	if (!ControlRig)
+	{
+		return false;
+	}
+	EditedTimes.Reset();
+
+	if (IsAnythingSelected())
+	{
+		IMovieScenePlayer* Player = GetSequencer().Get();
+		FFrameTime StartTime = GetSequencer()->GetLocalTime().Time;
+		FFrameRate TickResolution = GetSequencer()->GetFocusedTickResolution();
+		FMovieSceneContext MovieSceneContext = FMovieSceneContext(FMovieSceneEvaluationRange(StartTime, TickResolution), Player->GetPlaybackStatus()).SetHasJumped(true);
+		
+		Player->GetEvaluationTemplate().Evaluate(MovieSceneContext, *Player);
+		ControlRig->Evaluate_AnyThread();
+		return true;
+	}
+	return false;
+}
+
+bool FMovieSceneControlRigTransformTrail::HandleAltClick(FEditorViewportClient* InViewportClient, HMotionTrailProxy* Proxy, FInputClick Click)
+{
+	UMovieSceneControlRigParameterSection* Section = Cast<UMovieSceneControlRigParameterSection>(GetSection());
+	if (!Section)
+	{
+		return false;
+	}
+	UControlRig* ControlRig = Section->GetControlRig();
+	if (!ControlRig)
+	{
+		return false;
+	}
+	const FScopedTransaction Transaction(LOCTEXT("InsertKey", "Insert Key"));
+	Section->Modify();
+	ControlRig->Modify();
+	FFrameRate TickResolution = GetSequencer()->GetFocusedTickResolution();
+	FFrameRate DisplayResolution = GetSequencer()->GetFocusedDisplayRate();
+
+	FRigControlModifiedContext Context;
+	Context.SetKey = EControlRigSetKey::Always;
+	IMovieScenePlayer* Player = GetSequencer().Get();
+
+	//Snap to Frame. Assuming we want keys on frames.
+	FFrameTime GlobalTime = TickResolution.AsFrameTime(Proxy->Second);
+	FFrameTime DisplayTime = FFrameRate::TransformTime(GlobalTime, DisplayResolution, TickResolution);
+	DisplayTime = DisplayTime.RoundToFrame();
+	GlobalTime = FFrameRate::TransformTime(DisplayTime, TickResolution, DisplayResolution);
+
+	Context.LocalTime = TickResolution.AsSeconds(GlobalTime);
+	Context.KeyMask = (uint32)EControlRigContextChannelToKey::Translation;
+	FMovieSceneContext MovieSceneContext = FMovieSceneContext(FMovieSceneEvaluationRange(GlobalTime, TickResolution), Player->GetPlaybackStatus()).SetHasJumped(true);
+	GetSequencer()->GetEvaluationTemplate().Evaluate(MovieSceneContext, *Player);
+	ControlRig->Evaluate_AnyThread();
+	FTransform NewTransform(ControlRig->GetControlGlobalTransform(ControlName));
+	ControlRig->SetControlGlobalTransform(ControlName, NewTransform, true, Context, false);
+
+	//eval back at current time
+	FFrameTime StartTime = GetSequencer()->GetLocalTime().Time;
+	MovieSceneContext = FMovieSceneContext(FMovieSceneEvaluationRange(StartTime, TickResolution), Player->GetPlaybackStatus()).SetHasJumped(true);
+	Player->GetEvaluationTemplate().Evaluate(MovieSceneContext, *Player);
+	ControlRig->Evaluate_AnyThread();
+
+	//create new keys
+	KeyTool->BuildKeys();
+	FTrailKeyInfo* KeyInfo = KeyTool->FindKey(GlobalTime.RoundToFrame());
+	if (KeyInfo)
+	{
+		if (Click.bShiftIsDown == false && Click.bCtrlIsDown == false)
+		{
+			KeyTool->ClearSelection();;
+		}
+		KeyTool->SelectKeyInfo(KeyInfo);
+	}
+	return true;
+}
+
+} // namespace MovieScene
+} // namespace UE
+
+#undef LOCTEXT_NAMESPACE
\ No newline at end of file
diff --git a/Engine/Plugins/Editor/SequencerAnimTools/Source/SequencerAnimTools/Private/Sequencer/MovieSceneTransformTrail.h b/Engine/Plugins/Editor/SequencerAnimTools/Source/SequencerAnimTools/Private/Sequencer/MovieSceneTransformTrail.h
new file mode 100644
--- /dev/null
+++ b/Engine/Plugins/Editor/SequencerAnimTools/Source/SequencerAnimTools/Private/Sequencer/MovieSceneTransformTrail.h
@@ -0,0 +1,135 @@
+// Copyright Epic Games, Inc. All Rights Reserved.
+
+#pragma once
+
+#include "Trail.h"
+#include "MotionTrailMovieSceneKey.h"
+#include "TrajectoryDrawInfo.h"
+#include "EntitySystem/Interrogation/MovieSceneInterrogationLinker.h"
+
+class ISequencer;
+class UMovieScene3DTransformTrack;
+class UMovieScene3DTransformSection;
+class ULevelSequencePlayer;
+class UMovieSceneTrack;
+class USceneComponent;
+
+namespace UE
+{
+namespace SequencerAnimTools
+{
+
+
+class FMovieSceneTransformTrail : public FTrail
+{
+public:
+	FMovieSceneTransformTrail(USceneComponent* SceneComponent,const FLinearColor& InColor, const bool bInIsVisible, TWeakObjectPtr<UMovieSceneTrack> InWeakTrack, TSharedPtr<ISequencer> InSequencer, ULevelSequencePlayer* InLevelSequencePlayer);
+	// Begin FTrail interface
+	virtual ETrailCacheState UpdateTrail(const FSceneContext& InSceneContext) override;
+	virtual FTrajectoryCache* GetTrajectoryTransforms() override { return TrajectoryCache.Get(); }
+	virtual TRange<double> GetEffectiveRange() const override { return CachedEffectiveRange; }
+	virtual void Render(const FGuid& Guid, const FSceneView* View, FViewport* Viewport, FPrimitiveDrawInterface* PDI) override;
+	virtual void DrawHUD(FEditorViewportClient* ViewportClient, FViewport* Viewport, const FSceneView* View, FCanvas* Canvas) override;
+
+	virtual bool HandleClick(const FGuid& Guid, FEditorViewportClient* InViewportClient, HHitProxy* HitProxy, FInputClick Click) override;
+	virtual bool IsAnythingSelected() const override;
+	virtual bool IsAnythingSelected(FVector& OutVectorPosition)const override;
+	virtual bool IsTrailSelected() const override;
+	virtual bool StartTracking() override;
+	virtual bool ApplyDelta(const FVector& Pos, const FRotator& Rot, const FVector& WidgetLocation) override;
+	virtual bool EndTracking() override;
+	virtual void TranslateSelectedKeys(bool bRight) override;
+	virtual void DeleteSelectedKeys() override;
+	virtual void SelectNone() override;
+	virtual int32 GetChannelOffset()const { return 0; }
+	virtual bool GetEditedTimes(const FTrailHierarchy* TrailHierarchy,const FFrameNumber& LastFrame, TArray<FFrameNumber>& OutEditedTimes);
+	virtual void UpdateKeysInRange(const TRange<double>& ViewRange) override;
+	// End FTrail interface
+
+	TSharedPtr<ISequencer> GetSequencer() const { return WeakSequencer.Pin(); }
+	FGuid GetCachedHierarchyGuid() const { return CachedHierarchyGuid; }
+	UMovieSceneTrack* GetTrack() const { return WeakTrack.Get(); }
+	UMovieSceneSection* GetSection() const;
+
+public:
+	TWeakPtr<ISequencer> WeakSequencer;
+	TWeakObjectPtr<ULevelSequencePlayer> LevelSequencePlayer;
+	// Begin FMovieSceneTransformTrail interface
+	virtual void UpdateCacheTimes(const FSceneContext& InSceneContext) = 0;
+protected:
+	virtual bool HandleAltClick(FEditorViewportClient* InViewportClient, HMotionTrailProxy* Proxy, FInputClick Click);
+	// End FMovieSceneTransformTrail interface
+
+	// when editing interactively the tool may specify just a specify range to recalculate
+	void CalculateEditedTimes(const FTrailHierarchy* TrailHierarchy, const FFrameNumber& LastFrame);
+	void ClearEditedTiimes() { EditedTimes.Reset(); }
+
+	IMovieScenePlayer* MovieScenePlayer;
+	TRange<double> GetEffectiveSectionRange(int32 ChannelOffset) const;
+	TRange<double> CachedEffectiveRange;
+
+	bool bIsSelected;
+	FVector SelectedPos;
+
+	TUniquePtr<FMotionTraiMovieScenelKeyTool> KeyTool;
+	TUniquePtr<FArrayTrajectoryCache> TrajectoryCache;
+
+	FGuid LastTransformSectionSig;
+	FGuid CachedHierarchyGuid;
+	TWeakObjectPtr<UMovieSceneTrack> WeakTrack;
+
+	TOptional<TArray<FFrameNumber>> EditedTimes;
+
+
+};
+
+class FMovieSceneComponentTransformTrail : public FMovieSceneTransformTrail
+{
+public:
+
+	FMovieSceneComponentTransformTrail(USceneComponent* InComponent, const FLinearColor& InColor, const bool bInIsVisible, TWeakObjectPtr<UMovieSceneTrack> InWeakTrack, TSharedPtr<ISequencer> InSequencer, ULevelSequencePlayer* InLevelSequencePlayer)
+		: FMovieSceneTransformTrail(InComponent, InColor, bInIsVisible, InWeakTrack, InSequencer, InLevelSequencePlayer),
+		Component(InComponent)
+		
+	{
+	}
+
+
+private:
+	TWeakObjectPtr<USceneComponent> Component;
+	// Begin FMovieSceneTransformTrail interface
+	virtual void UpdateCacheTimes(const FSceneContext& InSceneContext) override;
+	virtual bool StartTracking() override;
+	virtual bool ApplyDelta(const FVector& Pos, const FRotator& Rot, const FVector& WidgetLocation) override;
+	virtual bool EndTracking() override;
+
+protected:
+	virtual bool HandleAltClick(FEditorViewportClient* InViewportClient, HMotionTrailProxy* Proxy, FInputClick Click) override;
+
+	// End FMovieSceneTransformTrail interface
+
+};
+
+
+class FMovieSceneControlRigTransformTrail : public FMovieSceneTransformTrail
+{
+public:
+
+	FMovieSceneControlRigTransformTrail(USceneComponent* SceneComponent, const FLinearColor& InColor, const bool bInIsVisible, TWeakObjectPtr<UMovieSceneTrack> InWeakTrack, TSharedPtr<ISequencer> InSequencer, ULevelSequencePlayer* InLevelSequencePlayer, const FName& InControlName);
+
+	virtual void UpdateCacheTimes(const FSceneContext& InSceneContext) override;
+	virtual bool StartTracking() override;
+	virtual bool ApplyDelta(const FVector& Pos, const FRotator& Rot, const FVector& WidgetLocation) override;
+	virtual bool EndTracking() override;
+	virtual int32 GetChannelOffset() const override;
+
+protected:
+	virtual bool HandleAltClick(FEditorViewportClient* InViewportClient, HMotionTrailProxy* Proxy, FInputClick Click) override;
+
+private:
+	FTransform EvaluateChannelsAtTime(TArrayView<FMovieSceneFloatChannel*> Channels, FFrameTime Time) const;
+	FName ControlName;
+};
+
+} // namespace MovieScene
+} // namespace UE
diff --git a/Engine/Plugins/Editor/SequencerAnimTools/Source/SequencerAnimTools/Private/Sequencer/SequencerTrailHierarchy.cpp b/Engine/Plugins/Editor/SequencerAnimTools/Source/SequencerAnimTools/Private/Sequencer/SequencerTrailHierarchy.cpp
new file mode 100644
--- /dev/null
+++ b/Engine/Plugins/Editor/SequencerAnimTools/Source/SequencerAnimTools/Private/Sequencer/SequencerTrailHierarchy.cpp
@@ -0,0 +1,804 @@
+// Copyright Epic Games, Inc. All Rights Reserved.
+
+#include "SequencerTrailHierarchy.h"
+
+#include "MovieSceneSequence.h"
+#include "ISequencer.h"
+#include "Tracks/MovieScene3DTransformTrack.h"
+#include "Tracks/MovieSceneSkeletalAnimationTrack.h"
+#include "MovieSceneSection.h"
+
+#include "ControlRig.h"
+#include "Sequencer/MovieSceneControlRigParameterTrack.h"
+#include "Sequencer/MovieSceneControlRigParameterSection.h"
+#include "AnimationBoneTrail.h"
+#include "LevelSequencePlayer.h"
+#include "Tools/MotionTrailOptions.h"
+#include "MovieSceneTransformTrail.h"
+#include "IControlRigObjectBinding.h"
+#include "ActorForWorldTransforms.h"
+#include "MovieSceneToolHelpers.h"
+#include "ControlRigSequencerEditorLibrary.h"
+#include "Engine/World.h"
+#include "Kismet/GameplayStatics.h"
+#include "EngineUtils.h"
+namespace UE
+{
+namespace SequencerAnimTools
+{
+
+void FSequencerTrailHierarchy::AssignLevelSequencePlayer()
+{
+	if (LevelSequencePlayer == nullptr)
+	{
+		LevelSequencePlayer = NewObject<ULevelSequencePlayer>();
+	}
+
+	UWorld* World = GEditor->GetEditorWorldContext().World();
+	FMovieSceneSequencePlaybackSettings Settings;
+	FLevelSequenceCameraSettings CameraSettings;
+	FMovieSceneSequenceIDRef Template = MovieSceneSequenceID::Root;
+	FMovieSceneSequenceTransform RootToLocalTransform;
+	TSharedPtr<ISequencer> Sequencer = WeakSequencer.Pin();
+	ULevelSequence* LevelSequence = Cast<ULevelSequence>(Sequencer->GetFocusedMovieSceneSequence());
+	LevelSequencePlayer->Initialize(LevelSequence, World->PersistentLevel, Settings, CameraSettings);
+	LevelSequencePlayer->State.AssignSequence(MovieSceneSequenceID::Root, *LevelSequence, *LevelSequencePlayer);
+}
+
+void FSequencerTrailHierarchy::Initialize()
+{
+	TSharedPtr<ISequencer> Sequencer = WeakSequencer.Pin();
+	if (!Sequencer)
+	{
+		return;
+	}
+	AssignLevelSequencePlayer();
+	UpdateViewAndEvalRange();
+
+	TArray<FGuid> SequencerSelectedObjects;
+	Sequencer->GetSelectedObjects(SequencerSelectedObjects);
+	UpdateSequencerBindings(SequencerSelectedObjects,
+		[this](UObject* Object, FTrail*, FGuid Guid) {
+		VisibilityManager.Selected.Add(Guid);
+	});
+	
+	OnSelectionChangedHandle = Sequencer->GetSelectionChangedObjectGuids().AddLambda([this](TArray<FGuid> NewSelection) {
+		
+		//if using old trails don't add new ones.
+		if (CVarUseOldSequencerMotionTrails->GetBool() == true)
+		{
+			return;
+		}		
+		TSharedPtr<ISequencer> Sequencer = WeakSequencer.Pin();
+		check(Sequencer);
+		TSet<FGuid> OldSelected = VisibilityManager.Selected;
+		TSet<FGuid> NewSelected;
+		
+		auto SetVisibleFunc = [this, &NewSelected](UObject* Object, FTrail* TrailPtr, FGuid Guid) {
+			NewSelected.Add(Guid);
+		};
+
+		UpdateSequencerBindings(NewSelection, SetVisibleFunc);
+		for (FGuid Guid : OldSelected)
+		{
+			if (NewSelected.Find(Guid) == nullptr)
+			{
+				RemoveTrailIfNotAlwaysVisible(Guid);
+			}
+		}
+		VisibilityManager.Selected = NewSelected;
+
+	});
+
+	OnViewOptionsChangedHandle = UMotionTrailToolOptions::GetTrailOptions()->OnDisplayPropertyChanged.AddLambda([this](FName PropertyName) {
+		if (PropertyName == GET_MEMBER_NAME_CHECKED(UMotionTrailToolOptions, EvalsPerFrame))
+		{
+			for (const TPair<FGuid, TUniquePtr<FTrail>>& GuidTrailPair : GetAllTrails())
+			{
+				GuidTrailPair.Value->ForceEvaluateNextTick();
+			}
+		}
+	});
+
+
+	GEngine->OnLevelActorAdded().AddRaw(this, &FSequencerTrailHierarchy::OnActorChangedSomehow);
+	GEngine->OnLevelActorDeleted().AddRaw(this, &FSequencerTrailHierarchy::OnActorChangedSomehow);
+	GEngine->OnActorMoved().AddRaw(this, &FSequencerTrailHierarchy::OnActorChangedSomehow);
+	GEngine->OnActorsMoved().AddRaw(this, &FSequencerTrailHierarchy::OnActorsChangedSomehow);
+	//GEditor->RegisterForUndo(this);
+
+}
+
+void FSequencerTrailHierarchy::Destroy()
+{
+	TSharedPtr<ISequencer> Sequencer = WeakSequencer.Pin();
+	if (Sequencer)
+	{
+		Sequencer->OnActorAddedToSequencer().Remove(OnActorAddedToSequencerHandle);
+		Sequencer->GetSelectionChangedObjectGuids().Remove(OnSelectionChangedHandle);
+		UMotionTrailToolOptions::GetTrailOptions()->OnDisplayPropertyChanged.Remove(OnViewOptionsChangedHandle);
+	}
+
+	for (const TPair<UMovieSceneControlRigParameterTrack*, FControlRigDelegateHandles>& SectionHandlesPair : ControlRigDelegateHandles)
+	{
+		UMovieSceneControlRigParameterTrack* Track = (SectionHandlesPair.Key);
+		if (Track && Track->GetControlRig())
+		{
+			URigHierarchy* RigHierarchy = Track->GetControlRig()->GetHierarchy();
+			Track->GetControlRig()->ControlSelected().Remove(SectionHandlesPair.Value.OnControlSelected);
+			RigHierarchy->OnModified().Remove(SectionHandlesPair.Value.OnHierarchyModified);
+		}
+	}
+
+	ObjectsTracked.Reset();
+	BonesTracked.Reset();
+	ControlsTracked.Reset();
+	AllTrails.Reset();
+
+	GEngine->OnLevelActorAdded().RemoveAll(this);
+	GEngine->OnLevelActorDeleted().RemoveAll(this);
+	GEngine->OnActorMoved().RemoveAll(this);
+	GEngine->OnActorsMoved().RemoveAll(this);
+	//GEditor->UnregisterForUndo(this);
+
+}
+
+void FSequencerTrailHierarchy::OnActorChangedSomehow(AActor* InActor)
+{
+	for (const TPair<FGuid, TUniquePtr<FTrail>>& GuidTrailPair : GetAllTrails())
+	{
+		USceneComponent* Component = Cast<USceneComponent>(GuidTrailPair.Value->GetOwner());
+		if (Component)
+		{
+			TArray<const UObject*> ParentActors;
+			UObject* Parent = Component->GetOwner();
+			MovieSceneToolHelpers::GetParents(ParentActors, Parent);
+			if (ParentActors.Contains(InActor))
+			{
+				GuidTrailPair.Value->ForceEvaluateNextTick();
+			}
+		}
+	}
+}
+
+void FSequencerTrailHierarchy::OnActorsChangedSomehow(TArray<AActor*>& InActors)
+{
+	for (const TPair<FGuid, TUniquePtr<FTrail>>& GuidTrailPair : GetAllTrails())
+	{
+		for (AActor* InActor : InActors)
+		{
+			USceneComponent* Component = Cast<USceneComponent>(GuidTrailPair.Value->GetOwner());
+			if (Component)
+			{
+				TArray<const UObject*> ParentActors;
+				UObject* Parent = Component->GetOwner();
+				MovieSceneToolHelpers::GetParents(ParentActors, Parent);
+				if (ParentActors.Contains(InActor))
+				{
+					GuidTrailPair.Value->ForceEvaluateNextTick();
+				}
+			}
+		}
+	}
+}
+
+void FSequencerTrailHierarchy::AddReferencedObjects(FReferenceCollector& Collector)
+{
+	Collector.AddReferencedObject(LevelSequencePlayer);
+}
+
+double FSequencerTrailHierarchy::GetSecondsPerSegment() const 
+{ 
+	double SecondsPerFrame =  WeakSequencer.Pin()->GetFocusedDisplayRate().AsInterval();
+	//evalsperframe will not be less than one
+	return SecondsPerFrame / double(UMotionTrailToolOptions::GetTrailOptions()->EvalsPerFrame);
+}
+
+FFrameNumber FSequencerTrailHierarchy::GetFramesPerSegment() const
+{
+	FFrameNumber FramesPerTick = FFrameRate::TransformTime(FFrameNumber(1), WeakSequencer.Pin()->GetFocusedDisplayRate(), 
+		WeakSequencer.Pin()->GetFocusedTickResolution()).RoundToFrame();
+	FramesPerTick.Value /= (UMotionTrailToolOptions::GetTrailOptions()->EvalsPerFrame);
+	return FramesPerTick;
+}
+
+FFrameNumber FSequencerTrailHierarchy::GetFramesPerFrame() const
+{
+	FFrameNumber FramesPerTick = FFrameRate::TransformTime(FFrameNumber(1), WeakSequencer.Pin()->GetFocusedDisplayRate(),
+		WeakSequencer.Pin()->GetFocusedTickResolution()).RoundToFrame();
+
+	return FramesPerTick;
+}
+
+void FSequencerTrailHierarchy::RemoveTrail(const FGuid& Key)
+{
+	FTrailHierarchy::RemoveTrail(Key);
+	if (UObject* const* FoundObject = ObjectsTracked.FindKey(Key))
+	{
+		ObjectsTracked.Remove(*FoundObject);
+	}
+	else
+	{
+		for (TPair<USkeletalMeshComponent*, TMap<FName, FGuid>>& CompMapPair : BonesTracked)
+		{
+			if (const FName* FoundBone = CompMapPair.Value.FindKey(Key))
+			{
+				CompMapPair.Value.Remove(*FoundBone);
+				return;
+			}
+		}
+		for (TPair<UControlRig*, TMap<FName, FGuid>>& CompMapPair : ControlsTracked)
+		{
+			if (const FName* FoundControl = CompMapPair.Value.FindKey(Key))
+			{
+				CompMapPair.Value.Remove(*FoundControl);
+				return;
+			}
+		}
+	}
+}
+struct FTrailControlTransforms
+{
+	FName ControlName;
+	FTrail* Trail;
+	TArray<FTransform> Transforms;
+};
+
+void FSequencerTrailHierarchy::UpdateControlRig(const FTrailEvaluateTimes& EvaluateTimes,UControlRig* ControlRig, TMap<FName, FGuid > &CompMapPair)
+{
+	if (TSharedPtr<IControlRigObjectBinding> ObjectBinding = ControlRig->GetObjectBinding())
+	{
+		TSharedPtr<ISequencer> Sequencer = WeakSequencer.Pin();
+		check(Sequencer);
+		FFrameRate TickResolution = Sequencer->GetFocusedTickResolution();
+		TArray<FFrameNumber> Frames;
+		bool bUseEditedTimes = false;
+
+		const FFrameTime LastFrame = EvaluateTimes.EvalTimes[EvaluateTimes.EvalTimes.Num() - 1] * TickResolution;
+		for (TPair<FName, FGuid >& Pair : CompMapPair)
+		{
+			if (AllTrails[Pair.Value]->GetEditedTimes(this,LastFrame.RoundToFrame(), Frames))
+			{
+				bUseEditedTimes = true;
+				break;
+			}
+		}
+		int32 Index = 0;
+		if (bUseEditedTimes == false)
+		{
+			//Todo mz make time in frames
+			Frames.SetNum(EvaluateTimes.EvalTimes.Num());
+			for (const double Time : EvaluateTimes.EvalTimes)
+			{
+				const FFrameTime TickTime = Time * TickResolution;
+				Frames[Index++] = TickTime.RoundToFrame();
+			}
+		}
+		UWorld* World = GEditor->GetEditorWorldContext().World();
+		ALevelSequenceActor* OutActor = nullptr;
+		FMovieSceneSequencePlaybackSettings Settings;
+		FLevelSequenceCameraSettings CameraSettings;
+		FMovieSceneSequenceIDRef Template = MovieSceneSequenceID::Root;
+		FMovieSceneSequenceTransform RootToLocalTransform;
+		IMovieScenePlayer* Player = LevelSequencePlayer;
+		ULevelSequence* LevelSequence = Cast<ULevelSequence>(Sequencer->GetFocusedMovieSceneSequence());
+
+		USceneComponent* Component = Cast<USceneComponent>(ObjectBinding->GetBoundObject());
+		if (Component)
+		{
+			AActor* Actor = Component->GetTypedOuter< AActor >();
+			if (Actor)
+			{
+				TArray<FTransform> ControlRigParentWorldTransforms;
+				FActorForWorldTransforms ControlRigActorSelection;
+				ControlRigActorSelection.Actor = Actor;
+
+				MovieSceneToolHelpers::GetActorWorldTransforms(Player, LevelSequence, Template, ControlRigActorSelection, Frames, ControlRigParentWorldTransforms);
+
+				TArray<FTrailControlTransforms> TrailControlTransforms;
+				TrailControlTransforms.SetNum(CompMapPair.Num());
+				int32 PairIndex = 0;
+				for (TPair<FName, FGuid >& Pair : CompMapPair)
+				{
+					TrailControlTransforms[PairIndex].Trail = AllTrails[Pair.Value].Get();
+					TrailControlTransforms[PairIndex].ControlName = Pair.Key;
+					TrailControlTransforms[PairIndex].Transforms.SetNum(Frames.Num());
+					++PairIndex;
+				}
+				for (Index = 0; Index < Frames.Num(); ++Index)
+				{
+					const FFrameNumber& FrameNumber = Frames[Index];
+					FFrameTime GlobalTime(FrameNumber);
+
+					FMovieSceneContext Context = FMovieSceneContext(FMovieSceneEvaluationRange(GlobalTime, TickResolution), Player->GetPlaybackStatus()).SetHasJumped(true);
+
+					Player->GetEvaluationTemplate().Evaluate(Context, *Player);
+					ControlRig->Evaluate_AnyThread();
+					for (FTrailControlTransforms& TrailControlTransform: TrailControlTransforms)
+					{
+						TrailControlTransform.Transforms[Index] = ControlRig->GetControlGlobalTransform(TrailControlTransform.ControlName) * ControlRigParentWorldTransforms[Index];
+						if (bUseEditedTimes)
+						{
+							double Sec = TickResolution.AsSeconds(FFrameTime(FrameNumber));
+							TrailControlTransform.Trail->GetTrajectoryTransforms()->Set(Sec, TrailControlTransform.Transforms[Index]);
+						}
+					}
+				}			
+				for (FTrailControlTransforms& TrailControlTransform : TrailControlTransforms)
+				{
+					if (bUseEditedTimes == false)
+					{
+						TrailControlTransform.Trail->GetTrajectoryTransforms()->SetTransforms(TrailControlTransform.Transforms, ControlRigParentWorldTransforms);
+					}
+					TrailControlTransform.Trail->UpdateKeysInRange(GetViewRange());
+				}			
+			}
+		}
+		if (Player)
+		{
+			const FFrameTime StartTime = Sequencer->GetLocalTime().Time;
+			FMovieSceneContext Context = FMovieSceneContext(FMovieSceneEvaluationRange(StartTime, TickResolution), Player->GetPlaybackStatus()).SetHasJumped(true);
+			Player->GetEvaluationTemplate().Evaluate(Context, *Player);
+			ControlRig->Evaluate_AnyThread();
+		}
+	}
+}
+void FSequencerTrailHierarchy::Update()
+{
+	const FDateTime UpdateStartTime = FDateTime::Now();
+
+	UpdateViewAndEvalRange();
+	FTrailHierarchy::Update();  //calculates EvalTimesArr
+	FTrailEvaluateTimes EvalTimes = FTrailEvaluateTimes(EvalTimesArr, SecondsPerSegment);
+
+	for (TPair<UControlRig*, TMap<FName, FGuid>>& CompMapPair : ControlsTracked)
+	{
+		bool bNeedToUpdateControlRig = false;
+		for (TPair<FName, FGuid>& Pair : CompMapPair.Value)
+		{
+			if (AllTrails[Pair.Value]->GetCacheState() == ETrailCacheState::Stale)
+			{
+				bNeedToUpdateControlRig = true;
+				break;
+			}
+		}
+		if (bNeedToUpdateControlRig == true)
+		{
+			UpdateControlRig(EvalTimes, CompMapPair.Key, CompMapPair.Value);
+		}
+	}
+
+	const FTimespan UpdateTimespan = FDateTime::Now() - UpdateStartTime;
+	TimingStats.Add("FSequencerTrailHierarchy::Update", UpdateTimespan);
+}
+
+void FSequencerTrailHierarchy::OnBoneVisibilityChanged(class USkeleton* Skeleton, const FName& BoneName, const bool bIsVisible)
+{
+	TSharedPtr<ISequencer> Sequencer = WeakSequencer.Pin();
+	check(Sequencer);
+
+	TArray<FGuid> SelectedSequencerGuids;
+	Sequencer->GetSelectedObjects(SelectedSequencerGuids);
+
+	// TODO: potentially expensive
+	for (const FGuid& SelectedGuid : SelectedSequencerGuids)
+	{
+		for (TWeakObjectPtr<> BoundObject : Sequencer->FindObjectsInCurrentSequence(SelectedGuid))
+		{
+			USkeletalMeshComponent* BoundComponent = Cast<USkeletalMeshComponent>(BoundObject.Get());
+			if (AActor* BoundActor = Cast<AActor>(BoundObject.Get()))
+			{
+				BoundComponent = BoundActor->FindComponentByClass<USkeletalMeshComponent>();
+			}
+
+			if (!BoundComponent || !BoundComponent->SkeletalMesh || !BoundComponent->SkeletalMesh->GetSkeleton() || !(BoundComponent->SkeletalMesh->GetSkeleton() == Skeleton) || !BonesTracked.Contains(BoundComponent))
+			{
+				continue;
+			}
+
+			const FGuid BoneTrailGuid = BonesTracked[BoundComponent][BoneName];
+			const int32 BoneIndex = Skeleton->GetReferenceSkeleton().FindBoneIndex(BoneName);
+			
+			if (bIsVisible)
+			{
+				VisibilityManager.VisibilityMask.Remove(BoneTrailGuid);
+				VisibilityManager.Selected.Add(BoneTrailGuid);
+			}
+			else
+			{
+				VisibilityManager.VisibilityMask.Add(BoneTrailGuid);
+				VisibilityManager.Selected.Remove(BoneTrailGuid);
+			}
+		}
+	}
+}
+
+void FSequencerTrailHierarchy::OnBindingVisibilityStateChanged(UObject* BoundObject, const EBindingVisibilityState VisibilityState)
+{
+	auto UpdateTrailVisibilityState = [this, VisibilityState](const FGuid& Guid) {
+		if (VisibilityState == EBindingVisibilityState::AlwaysVisible)
+		{
+			VisibilityManager.AlwaysVisible.Add(Guid);
+		}
+		else if (VisibilityState == EBindingVisibilityState::VisibleWhenSelected)
+		{
+			VisibilityManager.AlwaysVisible.Remove(Guid);
+		}
+	};
+
+	if (ObjectsTracked.Contains(BoundObject))
+	{
+		UpdateTrailVisibilityState(ObjectsTracked[BoundObject]);
+	}
+
+	USkeletalMeshComponent* SkelMeshComp = Cast<USkeletalMeshComponent>(BoundObject);
+	if (!SkelMeshComp)
+	{
+		return;
+	}
+
+	for (const TPair<FName, FGuid>& Pair : BonesTracked.FindRef(SkelMeshComp))
+	{
+		UpdateTrailVisibilityState(Pair.Value);
+	}
+	/* to do mz
+	for (const TPair<FName, FGuid>& Pair : ControlsTracked.FindRef(SkelMeshComp))
+	{
+		UpdateTrailVisibilityState(ControlsTracked[SkelMeshComp][Pair.Key]);
+	}
+	*/
+}
+
+void FSequencerTrailHierarchy::UpdateSequencerBindings(const TArray<FGuid>& SequencerBindings, TFunctionRef<void(UObject*, FTrail*, FGuid)> OnUpdated)
+{
+	//if using old trails don't add new ones.
+	if (CVarUseOldSequencerMotionTrails->GetBool() == true)
+	{
+		return;
+	}
+	const FDateTime StartTime = FDateTime::Now();
+
+	TSharedPtr<ISequencer> Sequencer = WeakSequencer.Pin();
+	check(Sequencer);
+
+	for (FGuid BindingGuid : SequencerBindings)
+	{
+		if (UMovieScene3DTransformTrack* TransformTrack = Sequencer->GetFocusedMovieSceneSequence()->GetMovieScene()->FindTrack<UMovieScene3DTransformTrack>(BindingGuid))
+		{
+			for (TWeakObjectPtr<> BoundObject : Sequencer->FindBoundObjects(BindingGuid, Sequencer->GetFocusedTemplateID()))
+			{
+				if (!BoundObject.IsValid())
+				{
+					continue;
+				}
+
+				USceneComponent* BoundComponent = Cast<USceneComponent>(BoundObject.Get());
+				if (AActor* BoundActor = Cast<AActor>(BoundObject.Get()))
+				{
+					BoundComponent = BoundActor->GetRootComponent();
+				}
+
+				if (!ObjectsTracked.Contains(BoundComponent))
+				{
+					AddComponentToHierarchy(BoundComponent, TransformTrack);
+				}
+
+				if (!ObjectsTracked.Contains(BoundComponent))
+				{
+					continue;
+				}
+
+				if (AllTrails.Contains(ObjectsTracked[BoundComponent]) && AllTrails[ObjectsTracked[BoundComponent]].IsValid())
+				{
+					OnUpdated(BoundComponent, AllTrails[ObjectsTracked[BoundComponent]].Get(), ObjectsTracked[BoundComponent]);
+				}
+			}
+		} // if TransformTrack
+		if (UMovieSceneSkeletalAnimationTrack* AnimTrack = Sequencer->GetFocusedMovieSceneSequence()->GetMovieScene()->FindTrack<UMovieSceneSkeletalAnimationTrack>(BindingGuid))
+		{
+			for (TWeakObjectPtr<> BoundObject : Sequencer->FindBoundObjects(BindingGuid, Sequencer->GetFocusedTemplateID()))
+			{
+				if (!BoundObject.IsValid())
+				{
+					continue;
+				}
+
+				USkeletalMeshComponent* BoundComponent = Cast<USkeletalMeshComponent>(BoundObject.Get());
+				if (AActor* BoundActor = Cast<AActor>(BoundObject.Get()))
+				{
+					BoundComponent = BoundActor->FindComponentByClass<USkeletalMeshComponent>();
+				}
+
+				if (!BoundComponent || !BoundComponent->SkeletalMesh || !BoundComponent->SkeletalMesh->GetSkeleton())
+				{
+					continue;
+				}
+
+				if (!ObjectsTracked.Contains(BoundComponent))
+				{
+					if (UMovieScene3DTransformTrack* TransformTrack = Sequencer->GetFocusedMovieSceneSequence()->GetMovieScene()->FindTrack<UMovieScene3DTransformTrack>(BindingGuid))
+					{
+						AddComponentToHierarchy(BoundComponent, TransformTrack);
+					}
+				}
+
+				if (!ObjectsTracked.Contains(BoundComponent))
+				{
+					continue;
+				}
+
+				if (!BonesTracked.Contains(BoundComponent))
+				{
+					//AddSkeletonToHierarchy(BoundComponent); //need to revisit skeletals
+				}
+
+				if (!BonesTracked.Contains(BoundComponent))
+				{
+					continue;
+				}
+
+				for (const TPair<FName, FGuid>& BoneNameGuidPair : BonesTracked[BoundComponent])
+				{
+					if (AllTrails.Contains(BoneNameGuidPair.Value) && AllTrails[BoneNameGuidPair.Value].IsValid())
+					{
+						OnUpdated(BoundComponent, AllTrails[BoneNameGuidPair.Value].Get(), BoneNameGuidPair.Value);
+					}
+				}
+			}
+		}
+		if (UMovieSceneControlRigParameterTrack* CRParameterTrack = Sequencer->GetFocusedMovieSceneSequence()->GetMovieScene()->FindTrack<UMovieSceneControlRigParameterTrack>(BindingGuid))
+		{
+			for (TWeakObjectPtr<> BoundObject : Sequencer->FindBoundObjects(BindingGuid, Sequencer->GetFocusedTemplateID()))
+			{
+				if (!BoundObject.IsValid())
+				{
+					continue;
+				}
+
+				USkeletalMeshComponent* BoundComponent = Cast<USkeletalMeshComponent>(BoundObject.Get());
+				if (AActor* BoundActor = Cast<AActor>(BoundObject.Get()))
+				{
+					BoundComponent = BoundActor->FindComponentByClass<USkeletalMeshComponent>();
+				}
+
+				if (!BoundComponent || !BoundComponent->SkeletalMesh || !BoundComponent->SkeletalMesh->GetSkeleton())
+				{
+					continue;
+				};
+
+				if (!ControlRigDelegateHandles.Contains(CRParameterTrack))
+				{
+					RegisterControlRigDelegates(BoundComponent, CRParameterTrack);
+					UControlRig* ControlRig = CRParameterTrack->GetControlRig();
+					if (ControlRig)
+					{
+						TArray<FName> Selected = ControlRig->CurrentControlSelection();
+						for (const FName& ControlName : Selected)
+						{
+							AddControlRigTrail(BoundComponent,ControlRig, CRParameterTrack, ControlName);
+						}
+					}
+				}
+			}
+		} // if ControlRigParameterTrack
+
+	}
+	const FTimespan Timespan = FDateTime::Now() - StartTime;
+	TimingStats.Add("FSequencerTrailHierarchy::UpdateSequencerBindings", Timespan);
+}
+
+
+void FSequencerTrailHierarchy::UpdateViewAndEvalRange()
+{
+	TSharedPtr<ISequencer> Sequencer = WeakSequencer.Pin();
+	check(Sequencer);
+
+	FFrameRate TickResolution = Sequencer->GetFocusedTickResolution();
+	FFrameRate DisplayRate = Sequencer->GetFocusedDisplayRate();
+
+	TRange<FFrameNumber> TickViewRange;
+	TRange<FFrameNumber> EvalViewRange  = Sequencer->GetFocusedMovieSceneSequence()->GetMovieScene()->GetPlaybackRange();
+	if (!UMotionTrailToolOptions::GetTrailOptions()->bShowFullTrail)
+	{
+		FFrameTime SequenceTime = Sequencer->GetLocalTime().Time;
+		const FFrameNumber TicksBefore = FFrameRate::TransformTime(FFrameNumber(UMotionTrailToolOptions::GetTrailOptions()->FramesBefore), DisplayRate, TickResolution).FloorToFrame();
+		const FFrameNumber TicksAfter = FFrameRate::TransformTime(FFrameNumber(UMotionTrailToolOptions::GetTrailOptions()->FramesAfter), DisplayRate, TickResolution).FloorToFrame();
+		TickViewRange = TRange<FFrameNumber>(SequenceTime.GetFrame() - TicksBefore, SequenceTime.GetFrame() + TicksAfter);
+	}
+	else
+	{
+		TickViewRange = Sequencer->GetFocusedMovieSceneSequence()->GetMovieScene()->GetPlaybackRange();
+	}
+
+	const double StartViewSeconds = TickResolution.AsSeconds(FFrameTime(TickViewRange.GetLowerBoundValue()));
+	const double EndViewSeconds = TickResolution.AsSeconds(FFrameTime(TickViewRange.GetUpperBoundValue()));
+	ViewRange = TRange<double>(StartViewSeconds,EndViewSeconds);
+
+	const double StartEvalSeconds = TickResolution.AsSeconds(FFrameTime(EvalViewRange.GetLowerBoundValue()));
+	const double EndEvalSeconds = TickResolution.AsSeconds(FFrameTime(EvalViewRange.GetUpperBoundValue()));
+	EvalRange = TRange<double>(StartEvalSeconds, EndEvalSeconds);
+
+	// snap view range to ticks per segment
+	//const double TicksBetween = FMath::Fmod(StartSeconds, GetSecondsPerSegment());
+	//ViewRange = TRange<double>(StartSeconds - TicksBetween, EndSeconds - TicksBetween);
+
+}
+
+void FSequencerTrailHierarchy::AddComponentToHierarchy(USceneComponent* CompToAdd, UMovieScene3DTransformTrack* TransformTrack)
+{
+	TSharedPtr<ISequencer> Sequencer = WeakSequencer.Pin();
+	check(Sequencer);
+	
+	const FGuid CurTrailGuid = ObjectsTracked.FindOrAdd(CompToAdd, FGuid::NewGuid());
+
+	TUniquePtr<FTrail> CurTrail = MakeUnique<FMovieSceneComponentTransformTrail>(CompToAdd,UMotionTrailToolOptions::GetTrailOptions()->TrailColor, false, TransformTrack, Sequencer, LevelSequencePlayer);
+	if (AllTrails.Contains(ObjectsTracked[CompToAdd])) 
+	{
+		AllTrails.Remove(ObjectsTracked[CompToAdd]);
+	}
+	CurTrail->ForceEvaluateNextTick();
+
+	AddTrail(ObjectsTracked[CompToAdd], MoveTemp(CurTrail));
+}
+
+void FSequencerTrailHierarchy::AddSkeletonToHierarchy(class USkeletalMeshComponent* CompToAdd)
+{
+	const FDateTime StartTime = FDateTime::Now();
+
+	TSharedPtr<ISequencer> Sequencer = WeakSequencer.Pin();
+	check(Sequencer);
+
+	TSharedPtr<FAnimTrajectoryCache> AnimTrajectoryCache = MakeShared<FAnimTrajectoryCache>(CompToAdd, Sequencer);
+	TMap<FName, FGuid>& BoneMap = BonesTracked.Add(CompToAdd, TMap<FName, FGuid>());
+	
+	USkeleton* MySkeleton = CompToAdd->SkeletalMesh->GetSkeleton();
+	const int32 NumBones = MySkeleton->GetReferenceSkeleton().GetNum();
+	for (int32 BoneIdx = 0; BoneIdx < NumBones; BoneIdx++)
+	{
+		int32 ParentBoneIndex = MySkeleton->GetReferenceSkeleton().GetParentIndex(BoneIdx);
+		const FName BoneName = MySkeleton->GetReferenceSkeleton().GetBoneName(BoneIdx);
+
+		const FGuid BoneGuid = FGuid::NewGuid();
+		BoneMap.Add(BoneName, BoneGuid);
+		VisibilityManager.VisibilityMask.Add(BoneGuid);
+
+		if (ParentBoneIndex != INDEX_NONE)
+		{
+			AllTrails.Add(BoneGuid, MakeUnique<FAnimationBoneTrail>(CompToAdd,UMotionTrailToolOptions::GetTrailOptions()->TrailColor, false, AnimTrajectoryCache, BoneName, false));
+
+		}
+		else // root bone
+		{
+			AllTrails.Add(BoneGuid, MakeUnique<FAnimationBoneTrail>(CompToAdd,UMotionTrailToolOptions::GetTrailOptions()->TrailColor, false, AnimTrajectoryCache, BoneName, true));
+		}
+	}
+
+	const FTimespan Timespan = FDateTime::Now() - StartTime;
+	TimingStats.Add("FSequencerTrailHierarchy::AddSkeletonToHierarchy", Timespan);
+}
+
+void FSequencerTrailHierarchy::AddControlRigTrail(USkeletalMeshComponent* Component, UControlRig* ControlRig, UMovieSceneControlRigParameterTrack* CRParameterTrack,const FName& ControlName)
+{
+	if (ControlsTracked.Find(ControlRig) == nullptr)
+	{
+		ControlsTracked.Add(ControlRig, TMap<FName, FGuid>());
+	}
+	if (!ControlsTracked[ControlRig].Contains(ControlName))
+	{
+		TSharedPtr<ISequencer> Sequencer = WeakSequencer.Pin();
+		if (Sequencer)
+		{
+			const FGuid ControlGuid = ControlsTracked.Find(ControlRig)->FindOrAdd(ControlName, FGuid::NewGuid());
+
+			TUniquePtr<FTrail> CurTrail = MakeUnique<FMovieSceneControlRigTransformTrail>(Component, UMotionTrailToolOptions::GetTrailOptions()->TrailColor, false, CRParameterTrack, Sequencer, LevelSequencePlayer, ControlName);
+			if (AllTrails.Contains(ControlGuid))
+			{
+				AllTrails.Remove(ControlGuid);
+				VisibilityManager.ControlSelected.Remove(ControlGuid);
+			}
+			AddTrail(ControlGuid, MoveTemp(CurTrail));
+			VisibilityManager.ControlSelected.Add(ControlGuid);
+		}
+	}
+	else
+	{
+		const FGuid* ControlGuid = ControlsTracked.Find(ControlRig)->Find(ControlName);
+		if (ControlGuid != nullptr)
+		{
+			VisibilityManager.ControlSelected.Add(*ControlGuid);
+		}
+	}
+}
+
+void FSequencerTrailHierarchy::RegisterControlRigDelegates(USkeletalMeshComponent* Component, UMovieSceneControlRigParameterTrack* CRParameterTrack)
+{
+	UControlRig* ControlRig = CRParameterTrack->GetControlRig();
+	URigHierarchy* RigHierarchy = ControlRig->GetHierarchy();
+	FControlRigDelegateHandles& DelegateHandles = ControlRigDelegateHandles.Add(CRParameterTrack);
+	DelegateHandles.OnControlSelected = ControlRig->ControlSelected().AddLambda([this, Component,CRParameterTrack](UControlRig* ControlRig, FRigControlElement* ControlElement, bool bSelected)
+		{
+			if (ControlElement->Settings.ControlType != ERigControlType::Transform &&
+				ControlElement->Settings.ControlType != ERigControlType::TransformNoScale &&
+				ControlElement->Settings.ControlType != ERigControlType::EulerTransform)
+			{
+				return;
+			}
+
+			if (bSelected)
+			{
+				AddControlRigTrail(Component,ControlRig,CRParameterTrack, ControlElement->GetName());
+			}
+
+			if (ControlsTracked.Find(ControlRig) != nullptr)
+			{
+				if (ControlsTracked[ControlRig].Contains(ControlElement->GetName()))
+				{
+					if (bSelected == false)
+					{
+						const FGuid TrailGuid = ControlsTracked[ControlRig][ControlElement->GetName()];
+						VisibilityManager.ControlSelected.Remove(TrailGuid);
+						RemoveTrailIfNotAlwaysVisible(TrailGuid);
+					}
+				}
+			}
+
+			//check to see if the seleced control rig is sill selected
+			for (TPair<UControlRig*, TMap<FName, FGuid>>& CompMapPair : ControlsTracked)
+			{
+				UControlRig* TrackedControlRig = CompMapPair.Key;
+				for (TPair<FName, FGuid> NameGuid : CompMapPair.Value)
+				{
+					if (TrackedControlRig->IsControlSelected(NameGuid.Key) == false)
+					{
+						const FGuid TrailGuid = NameGuid.Value;
+						VisibilityManager.ControlSelected.Remove(TrailGuid);
+						RemoveTrailIfNotAlwaysVisible(TrailGuid);
+					}
+				}
+			}
+		}
+
+			
+	);
+	
+	DelegateHandles.OnHierarchyModified = RigHierarchy->OnModified().AddLambda(
+		[this, RigHierarchy, ControlRig](ERigHierarchyNotification InNotif, URigHierarchy* InHierarchy, const FRigBaseElement* InElement) {
+			TSharedPtr<ISequencer> Sequencer = WeakSequencer.Pin();
+			check(Sequencer);
+
+			const FRigControlElement* ControlElement = Cast<FRigControlElement>(InElement); 
+			if(ControlElement == nullptr)
+			{
+				return;
+			}
+
+			
+			if(InNotif == ERigHierarchyNotification::ElementRemoved)
+			{
+				if (!ControlsTracked.Contains(ControlRig) || !ControlsTracked[ControlRig].Contains(ControlElement->GetName())) // We only care about controls
+				{
+					return;
+				}
+				
+				const FGuid TrailGuid = ControlsTracked[ControlRig][ControlElement->GetName()];
+				RemoveTrail(TrailGuid);
+			}
+            else if(InNotif == ERigHierarchyNotification::ElementRenamed)
+            {
+            	const FName OldName = InHierarchy->GetPreviousName(ControlElement->GetKey());
+
+            	if (!ControlsTracked.Contains(ControlRig) || !ControlsTracked[ControlRig].Contains(OldName))
+				{
+					return;
+				}
+
+				const FGuid TempTrailGuid = ControlsTracked[ControlRig][OldName];
+				ControlsTracked[ControlRig].Remove(OldName);
+				ControlsTracked[ControlRig].Add(ControlElement->GetName(), TempTrailGuid);
+            }
+        }
+    );
+}
+
+} // namespace MovieScene
+} // namespace UE
diff --git a/Engine/Plugins/Editor/SequencerAnimTools/Source/SequencerAnimTools/Private/Sequencer/SequencerTrailHierarchy.h b/Engine/Plugins/Editor/SequencerAnimTools/Source/SequencerAnimTools/Private/Sequencer/SequencerTrailHierarchy.h
new file mode 100644
--- /dev/null
+++ b/Engine/Plugins/Editor/SequencerAnimTools/Source/SequencerAnimTools/Private/Sequencer/SequencerTrailHierarchy.h
@@ -0,0 +1,120 @@
+// Copyright Epic Games, Inc. All Rights Reserved.
+
+#pragma once
+
+#include "TrailHierarchy.h"
+#include "ISequencer.h"
+#include "Rigs/RigHierarchyDefines.h"
+#include "UObject/GCObject.h"
+
+
+class ISequencer;
+class USceneComponent;
+class USkeletalMeshComponent;
+class USkeleton;
+class UMovieSceneSection;
+class UMovieScene3DTransformTrack;
+class UMovieSceneControlRigParameterTrack;
+class UMovieSceneControlRigParameterSection;
+class UControlRig;
+struct FRigHierarchyContainer;
+class ULevelSequencePlayer;
+
+namespace UE
+{
+namespace SequencerAnimTools
+{
+
+enum class EBindingVisibilityState
+{
+	AlwaysVisible,
+	VisibleWhenSelected
+};
+
+class FSequencerTrailHierarchy : public FTrailHierarchy, public FGCObject
+{
+public:
+	FSequencerTrailHierarchy(TWeakPtr<ISequencer> InWeakSequencer)
+		: FTrailHierarchy()
+		, WeakSequencer(InWeakSequencer)
+		, ObjectsTracked()
+		, ControlsTracked()
+		, LevelSequencePlayer(nullptr)
+		, HierarchyRenderer(MakeUnique<FTrailHierarchyRenderer>(this))
+		, OnActorAddedToSequencerHandle()
+		, OnSelectionChangedHandle()
+		, OnViewOptionsChangedHandle()
+		, ControlRigDelegateHandles()
+	{
+	}
+
+	// FTrailHierarchy interface
+	virtual void Initialize() override;
+	virtual void Destroy() override;
+	virtual ITrailHierarchyRenderer* GetRenderer() const override { return HierarchyRenderer.Get(); }
+	virtual double GetSecondsPerFrame() const override { return 1.0 / WeakSequencer.Pin()->GetFocusedDisplayRate().AsDecimal(); }
+	virtual double GetSecondsPerSegment() const override;
+	virtual FFrameNumber GetFramesPerFrame() const override;
+	virtual FFrameNumber GetFramesPerSegment() const override;
+
+	virtual void RemoveTrail(const FGuid& Key) override;
+	virtual void Update() override;
+	// End FTrailHierarchy interface
+
+	// Begin FGCObject interface, this holds reference to a common ULevelSequencePlayer we use for evaluation.
+	virtual void AddReferencedObjects(FReferenceCollector& Collector) override;
+	virtual FString GetReferencerName() const override
+	{
+		return TEXT("UE::SequencerAnimTools::FMovieSceneTransformTrail");
+	}
+	// End FGCObject interface
+	const TMap<UObject*, FGuid>& GetObjectsTracked() const { return ObjectsTracked; }
+	const TMap<USkeletalMeshComponent*, TMap<FName, FGuid>>& GetBonesTracked() const { return BonesTracked; }
+	const TMap<UControlRig*, TMap<FName, FGuid>>& GetControlsTracked() const { return ControlsTracked; }
+
+	void OnBoneVisibilityChanged(USkeleton* Skeleton, const FName& BoneName, const bool bIsVisible);
+	void OnBindingVisibilityStateChanged(UObject* BoundObject, const EBindingVisibilityState VisibilityState);
+	//called when added moved, deleted, added, all of which may need a refresh of a motion trail
+	void OnActorChangedSomehow(AActor* InActor);
+	void OnActorsChangedSomehow(TArray<AActor*>& InActors);	
+
+private:
+	void AddControlRigTrail(USkeletalMeshComponent* Component,UControlRig* ControlRig, UMovieSceneControlRigParameterTrack* CRTrack, const FName& ControlName);
+
+	void UpdateControlRig(const FTrailEvaluateTimes& EvalTimes, UControlRig* ControlRig, TMap<FName, FGuid >& CompMapPair);
+
+	void UpdateSequencerBindings(const TArray<FGuid>& SequencerBindings, TFunctionRef<void(UObject*, FTrail*, FGuid)> OnUpdated);
+	void UpdateViewAndEvalRange();
+
+	void AddComponentToHierarchy(USceneComponent* CompToAdd, UMovieScene3DTransformTrack* TransformTrack);
+	void AddSkeletonToHierarchy(USkeletalMeshComponent* CompToAdd);
+
+	void RegisterControlRigDelegates(USkeletalMeshComponent* Component, UMovieSceneControlRigParameterTrack* CRParameterTrack);
+	void AssignLevelSequencePlayer();
+
+	TWeakPtr<ISequencer> WeakSequencer;
+	TMap<UObject*, FGuid> ObjectsTracked;
+	TMap<USkeletalMeshComponent*, TMap<FName, FGuid>> BonesTracked; 
+	TMap<UControlRig*, TMap<FName, FGuid>> ControlsTracked; 
+
+	// TODO: components can have multiple rigs so make this a map from sections to controls instead. However, this is only part of a larger problem of handling blending
+
+	ULevelSequencePlayer* LevelSequencePlayer; //gco object
+	TUniquePtr<FTrailHierarchyRenderer> HierarchyRenderer;
+
+	FDelegateHandle OnActorAddedToSequencerHandle;
+	FDelegateHandle OnSelectionChangedHandle;
+	FDelegateHandle OnViewOptionsChangedHandle;
+
+	struct FControlRigDelegateHandles
+	{
+		FDelegateHandle OnHierarchyModified;
+		FDelegateHandle OnControlSelected;
+	};
+	TMap<UMovieSceneControlRigParameterTrack*, FControlRigDelegateHandles> ControlRigDelegateHandles;
+
+
+};
+
+} // namespace MovieScene
+} // namespace UE
diff --git a/Engine/Plugins/Editor/SequencerAnimTools/Source/SequencerAnimTools/Private/SequencerAnimToolsModule.cpp b/Engine/Plugins/Editor/SequencerAnimTools/Source/SequencerAnimTools/Private/SequencerAnimToolsModule.cpp
new file mode 100644
--- /dev/null
+++ b/Engine/Plugins/Editor/SequencerAnimTools/Source/SequencerAnimTools/Private/SequencerAnimToolsModule.cpp
@@ -0,0 +1,123 @@
+// Copyright Epic Games, Inc. All Rights Reserved.
+
+#include "SequencerAnimToolsModule.h"
+#include "EditorModeManager.h"
+#include "Editor.h"
+#include "EditorModeRegistry.h"
+#include "LevelEditor.h"
+#include "EdModeInteractiveToolsContext.h"
+#include "UnrealEdGlobals.h"
+
+#include "BaseGizmos\TransformGizmoUtil.h"
+#include "MotionTrailTool.h"
+#include "EditPivotTool.h"
+#include "Tools\MotionTrailOptions.h"
+
+#define LOCTEXT_NAMESPACE "FSequencerAnimToolsModule"
+
+namespace UE
+{
+namespace SequencerAnimTools
+{
+
+void FSequencerAnimToolsModule::StartupModule()
+{
+	FLevelEditorModule& LevelEditor = FModuleManager::Get().LoadModuleChecked<FLevelEditorModule>("LevelEditor");
+	LevelEditor.OnLevelEditorCreated().AddRaw(this, &FSequencerAnimToolsModule::OnLevelEditorCreated);
+	UMotionTrailToolOptions* MotionTrailOptions = GetMutableDefault<UMotionTrailToolOptions>();
+	if (MotionTrailOptions)
+	{
+		MotionTrailOptions->OnDisplayPropertyChanged.AddRaw(this, &FSequencerAnimToolsModule::OnMotionTralOptionChanged);
+	}
+}
+
+void FSequencerAnimToolsModule::ShutdownModule()
+{
+	FLevelEditorModule& LevelEditor = FModuleManager::Get().LoadModuleChecked<FLevelEditorModule>("LevelEditor");
+	LevelEditor.OnLevelEditorCreated().RemoveAll(this);
+	/* can't do this for some reason the object is already dead (but not nullptr so can't check)
+	UMotionTrailToolOptions* MotionTrailOptions = GetMutableDefault<UMotionTrailToolOptions>();
+	if (MotionTrailOptions)
+	{
+		MotionTrailOptions->OnDisplayPropertyChanged.RemoveAll(this);
+	}
+	*/
+}
+
+void FSequencerAnimToolsModule::OnLevelEditorCreated(TSharedPtr<ILevelEditor> InLevelEditor)
+{
+	LevelEditorPtr = InLevelEditor;
+	InLevelEditor->GetEditorModeManager().OnEditorModeIDChanged().AddLambda([this](const FEditorModeID& InModeID, bool IsEnteringMode)
+		{
+			FEditorModeID ModeID = TEXT("EM_SequencerMode");
+			if (ModeID == InModeID && LevelEditorPtr.IsValid())
+			{
+				//don't re-register, will hit an ensure
+				if (AlreadyRegisteredTools.Contains(LevelEditorPtr.Pin().Get()) == false)
+				{
+					UEdMode* EdMode = LevelEditorPtr.Pin()->GetEditorModeManager().GetActiveScriptableMode(ModeID);
+
+					// Register the factories and gizmos if we are entering the new Gizmo Mode
+					if (IsEnteringMode && EdMode)
+					{
+						LevelEditorPtr.Pin()->GetEditorModeManager().GetInteractiveToolsContext()->ToolManager->RegisterToolType(TEXT("SequencerMotionTrail"), NewObject<UMotionTrailToolBuilder>(EdMode));
+						LevelEditorPtr.Pin()->GetEditorModeManager().GetInteractiveToolsContext()->ToolManager->RegisterToolType(TEXT("SequencerPivotTool"), NewObject<USequencerPivotToolBuilder>(EdMode));
+
+						UE::TransformGizmoUtil::RegisterTransformGizmoContextObject(LevelEditorPtr.Pin()->GetEditorModeManager().GetInteractiveToolsContext());
+
+					}
+					AlreadyRegisteredTools.Add(LevelEditorPtr.Pin().Get());
+				}
+			}
+		});
+}
+
+void FSequencerAnimToolsModule::OnMotionTralOptionChanged(FName PropertyName)
+{
+	UMotionTrailToolOptions* MotionTrailOptions = GetMutableDefault<UMotionTrailToolOptions>();
+	if (MotionTrailOptions)
+	{
+		if (PropertyName == FName("bShowTrails") && LevelEditorPtr.IsValid())
+		{
+			static bool bIsChangingTrail = false;
+
+			if (bIsChangingTrail == false)
+			{
+				bIsChangingTrail = true;
+				if (FLevelEditorSequencerIntegration::Get().GetSequencers().Num() > 0)
+				{
+
+					if (MotionTrailOptions->bShowTrails)
+					{
+						LevelEditorPtr.Pin()->GetEditorModeManager().GetInteractiveToolsContext()->ToolManager->SelectActiveToolType(EToolSide::Left, TEXT("SequencerMotionTrail"));
+						LevelEditorPtr.Pin()->GetEditorModeManager().GetInteractiveToolsContext()->ToolManager->ActivateTool(EToolSide::Left);
+					}
+					else
+					{
+						if (Cast<UMotionTrailTool>(LevelEditorPtr.Pin()->GetEditorModeManager().GetInteractiveToolsContext()->ToolManager->ActiveLeftTool))
+						{
+							LevelEditorPtr.Pin()->GetEditorModeManager().GetInteractiveToolsContext()->ToolManager->DeactivateTool(EToolSide::Left, EToolShutdownType::Completed);
+						}
+					}
+					
+				}
+				else
+				{
+					MotionTrailOptions->bShowTrails = false;
+				}
+				bIsChangingTrail = false;
+			}
+		}
+	}
+}
+
+
+
+
+
+} // namespace SequencerAnimTools
+} // namespace UE
+
+#undef LOCTEXT_NAMESPACE
+	
+IMPLEMENT_MODULE(UE::SequencerAnimTools::FSequencerAnimToolsModule, SequencerAnimTools)
\ No newline at end of file
diff --git a/Engine/Plugins/Editor/SequencerAnimTools/Source/SequencerAnimTools/Private/SequencerToolsEditMode.cpp b/Engine/Plugins/Editor/SequencerAnimTools/Source/SequencerAnimTools/Private/SequencerToolsEditMode.cpp
new file mode 100644
--- /dev/null
+++ b/Engine/Plugins/Editor/SequencerAnimTools/Source/SequencerAnimTools/Private/SequencerToolsEditMode.cpp
@@ -0,0 +1,64 @@
+// Copyright Epic Games, Inc. All Rights Reserved.
+
+#include "SequencerToolsEditMode.h"
+#include "EditorModeManager.h"
+#include "EdModeInteractiveToolsContext.h"
+#include "EdMode.h"
+#include "LevelEditor.h"
+#include "Modules/ModuleManager.h"
+#include "EditorModeManager.h"
+#include "InteractiveToolManager.h"
+#include "BaseSequencerAnimTool.h"
+
+#define LOCTEXT_NAMESPACE "SequencerAnimTools"
+
+
+FEditorModeID USequencerToolsEditMode::ModeName = TEXT("SequencerToolsEditMode");
+
+
+USequencerToolsEditMode::USequencerToolsEditMode() 
+{
+	Info = FEditorModeInfo(
+		ModeName,
+		LOCTEXT("ModeName", "Sequencer Tools"),
+		FSlateIcon(),
+		false
+	);
+}
+
+USequencerToolsEditMode::~USequencerToolsEditMode()
+{
+
+}
+
+bool USequencerToolsEditMode::IsCompatibleWith(FEditorModeID OtherModeID) const
+{
+	return OtherModeID == FName(TEXT("EM_SequencerMode"), FNAME_Find) || OtherModeID == FName(TEXT("EditMode.ControlRig"), FNAME_Find) || OtherModeID == FName(TEXT("EditMode.ControlRigEditor"), FNAME_Find);
+}
+
+//If we have one of our own active tools we pass input to it's commands.
+bool USequencerToolsEditMode::InputKey(FEditorViewportClient* InViewportClient, FViewport* InViewport, FKey InKey, EInputEvent InEvent) 
+{
+	if (InEvent != IE_Released)
+	{
+		//	MZ doesn't seem needed, need more testing  TGuardValue<FEditorViewportClient*> ViewportGuard(CurrentViewportClient, InViewportClient);
+		if (FLevelEditorModule* LevelEditorModule = FModuleManager::GetModulePtr<FLevelEditorModule>(TEXT("LevelEditor")))
+		{
+			TSharedPtr<ILevelEditor> LevelEditor = LevelEditorModule->GetLevelEditorInstance().Pin();
+
+			if (LevelEditor.IsValid())
+			{
+				if (IBaseSequencerAnimTool* Tool = Cast<IBaseSequencerAnimTool>(LevelEditor->GetEditorModeManager().GetInteractiveToolsContext()->ToolManager->GetActiveTool(EToolSide::Left)))
+				{
+					if (Tool->ProcessCommandBindings(InKey, (InEvent == IE_Repeat)))
+					{
+						return true;
+					}
+				}
+			}
+		}
+	}	
+	return UBaseLegacyWidgetEdMode::InputKey(InViewportClient, InViewport, InKey, InEvent);
+}
+
+#undef LOCTEXT_NAMESPACE
\ No newline at end of file
diff --git a/Engine/Plugins/Editor/SequencerAnimTools/Source/SequencerAnimTools/Private/Trail.cpp b/Engine/Plugins/Editor/SequencerAnimTools/Source/SequencerAnimTools/Private/Trail.cpp
new file mode 100644
--- /dev/null
+++ b/Engine/Plugins/Editor/SequencerAnimTools/Source/SequencerAnimTools/Private/Trail.cpp
@@ -0,0 +1,7 @@
+// Copyright Epic Games, Inc. All Rights Reserved.
+
+#include "Trail.h"
+#include "TrailHierarchy.h"
+
+//MZ TODO LOOK AT REIMPLEMENTING CONSTANT TRAILS
+
diff --git a/Engine/Plugins/Editor/SequencerAnimTools/Source/SequencerAnimTools/Private/TrailHierarchy.cpp b/Engine/Plugins/Editor/SequencerAnimTools/Source/SequencerAnimTools/Private/TrailHierarchy.cpp
new file mode 100644
--- /dev/null
+++ b/Engine/Plugins/Editor/SequencerAnimTools/Source/SequencerAnimTools/Private/TrailHierarchy.cpp
@@ -0,0 +1,558 @@
+// Copyright Epic Games, Inc. All Rights Reserved.
+
+#include "TrailHierarchy.h"
+#include "TrajectoryDrawInfo.h"
+#include "Containers/Queue.h"
+#include "CanvasItem.h"
+#include "CanvasTypes.h"
+#include "LevelEditor.h"
+#include "Framework/Application/SlateApplication.h"
+#include "EditorModeManager.h"
+#include "Tools/MotionTrailOptions.h"
+
+#define LOCTEXT_NAMESPACE "MotionTrailEditorMode"
+
+namespace UE
+{
+namespace SequencerAnimTools
+{
+
+IMPLEMENT_HIT_PROXY(HBaseTrailProxy, HHitProxy);
+IMPLEMENT_HIT_PROXY(HMotionTrailProxy, HBaseTrailProxy);
+
+void FTrailHierarchyRenderer::Render(const FSceneView* View, FViewport* Viewport, FPrimitiveDrawInterface* PDI)
+{
+	const FDateTime RenderStartTime = FDateTime::Now();
+	
+	const int32 NumEvalTimes = int32(OwningHierarchy->GetViewRange().Size<double>() / OwningHierarchy->GetSecondsPerSegment());
+	const int32 NumLinesReserveSize = int32(NumEvalTimes * OwningHierarchy->GetAllTrails().Num() * 1.3);
+	PDI->AddReserveLines(SDPG_Foreground, NumLinesReserveSize);
+
+	TArray<FGuid> TrailsToRemove; //if trail is not visible we just get rid of it
+	// <parent, current>
+	for (const TPair<FGuid, TUniquePtr<FTrail>>& GuidTrailPair : OwningHierarchy->GetAllTrails())
+	{
+		FTrajectoryDrawInfo* CurDrawInfo = GuidTrailPair.Value->GetDrawInfo();
+		if (CurDrawInfo && OwningHierarchy->GetVisibilityManager().IsTrailVisible(GuidTrailPair.Key, GuidTrailPair.Value.Get()))
+		{
+			FTrajectoryDrawInfo::FDisplayContext DisplayContext = {
+				GuidTrailPair.Key,
+				FTrailScreenSpaceTransform(View, Viewport),
+				UMotionTrailToolOptions::GetTrailOptions()->SecondsPerTick,
+				OwningHierarchy->GetViewRange(),
+				OwningHierarchy->GetSecondsPerSegment()
+			};
+
+			const bool bHitTesting = PDI && PDI->IsHitTesting();
+			TArray<FVector> PointsToDraw;
+			TArray<double> TrajectoryTimes;
+			CurDrawInfo->GetTrajectoryPointsForDisplay(DisplayContext, PointsToDraw, TrajectoryTimes);
+			FLinearColor Color = OwningHierarchy->GetVisibilityManager().IsTrailAlwaysVisible(GuidTrailPair.Key) ? CurDrawInfo->GetColor() : UMotionTrailToolOptions::GetTrailOptions()->TrailColor;
+			Color = GuidTrailPair.Value->IsTrailSelected() ? FLinearColor(1.0f, 1.0f, 0.0f) : Color;
+			if (PointsToDraw.Num() > 1)
+			{
+				const float TrailThickness = UMotionTrailToolOptions::GetTrailOptions()->TrailThickness;
+				FVector LastPoint = PointsToDraw[0];
+				for (int32 Idx = 1; Idx < PointsToDraw.Num(); Idx++)
+				{
+					if (bHitTesting)
+					{
+						const double CurSecond = TrajectoryTimes[Idx - 1];
+						PDI->SetHitProxy(new HMotionTrailProxy(GuidTrailPair.Key, LastPoint, CurSecond));
+					}
+					const FVector CurPoint = PointsToDraw[Idx];
+					PDI->DrawLine(LastPoint, CurPoint, Color, SDPG_Foreground,TrailThickness);
+					LastPoint = CurPoint;
+					if (bHitTesting)
+					{
+						PDI->SetHitProxy(nullptr);
+					}
+				}
+			}
+			GuidTrailPair.Value->Render(GuidTrailPair.Key,View,Viewport,PDI);
+
+		}
+		else
+		{
+			TrailsToRemove.Add(GuidTrailPair.Key);
+		}
+
+	}
+	for (const FGuid& Key : TrailsToRemove)
+	{
+		OwningHierarchy->RemoveTrail(Key);
+	}
+	const FTimespan RenderTimespan = FDateTime::Now() - RenderStartTime;
+	OwningHierarchy->GetTimingStats().Add("FTrailHierarchyRenderer::Render", RenderTimespan);
+}
+
+void FTrailHierarchyRenderer::DrawHUD(FEditorViewportClient* ViewportClient, FViewport* Viewport, const FSceneView* View, FCanvas* Canvas)
+{
+
+	if (UMotionTrailToolOptions::GetTrailOptions()->bShowTicks == false)
+	{
+		for (const TPair<FGuid, TUniquePtr<FTrail>>& GuidTrailPair : OwningHierarchy->GetAllTrails())
+		{
+			GuidTrailPair.Value->DrawHUD( ViewportClient, Viewport, View, Canvas);
+		}
+		return;
+	}	
+
+	const FDateTime DrawHUDStartTime = FDateTime::Now();
+
+	const double SecondsPerTick = UMotionTrailToolOptions::GetTrailOptions()->bLockTicksToFrames ? OwningHierarchy->GetSecondsPerFrame() : UMotionTrailToolOptions::GetTrailOptions()->SecondsPerTick;
+	const int32 PredictedNumTicks = int32((OwningHierarchy->GetViewRange().Size<double>() / SecondsPerTick) * OwningHierarchy->GetAllTrails().Num() * 1.3); // Multiply by 1.3 to be safe
+	Canvas->GetBatchedElements(FCanvas::EElementType::ET_Line)->AddReserveLines(PredictedNumTicks);
+
+	for(const TPair<FGuid, TUniquePtr<FTrail>>& GuidTrailPair : OwningHierarchy->GetAllTrails())
+	{
+		FTrajectoryDrawInfo* CurDrawInfo = GuidTrailPair.Value->GetDrawInfo();
+		if (CurDrawInfo && OwningHierarchy->GetVisibilityManager().IsTrailVisible(GuidTrailPair.Key, GuidTrailPair.Value.Get()))
+		{
+			FTrajectoryDrawInfo::FDisplayContext DisplayContext = {
+				GuidTrailPair.Key,
+				FTrailScreenSpaceTransform(View, Viewport, ViewportClient->GetDPIScale()),
+				SecondsPerTick,
+				OwningHierarchy->GetViewRange(),
+				OwningHierarchy->GetSecondsPerSegment()
+
+			};
+
+			TArray<FVector2D> Ticks, TickNormals;
+			CurDrawInfo->GetTickPointsForDisplay(DisplayContext, Ticks, TickNormals);
+			const FLinearColor Color = OwningHierarchy->GetVisibilityManager().IsTrailAlwaysVisible(GuidTrailPair.Key) ? CurDrawInfo->GetColor() : UMotionTrailToolOptions::GetTrailOptions()->TickColor;
+			for (int32 Idx = 0; Idx < Ticks.Num(); Idx++)
+			{
+				const FVector2D StartPoint = Ticks[Idx] - TickNormals[Idx] * UMotionTrailToolOptions::GetTrailOptions()->TickSize;
+				const FVector2D EndPoint = Ticks[Idx] + TickNormals[Idx] * UMotionTrailToolOptions::GetTrailOptions()->TickSize;
+				FCanvasLineItem LineItem = FCanvasLineItem(StartPoint, EndPoint);
+				LineItem.SetColor(Color);
+				Canvas->DrawItem(LineItem);
+			}
+		}
+		GuidTrailPair.Value->DrawHUD(ViewportClient, Viewport, View, Canvas);
+
+	}
+
+	const FTimespan DrawHUDTimespan = FDateTime::Now() - DrawHUDStartTime;
+	OwningHierarchy->GetTimingStats().Add("FTrailHierarchyRenderer::DrawHUD", DrawHUDTimespan);
+}
+
+class STrailOptionsPopup : public SCompoundWidget
+{
+public:
+
+	SLATE_BEGIN_ARGS(STrailOptionsPopup)
+	{}
+
+	SLATE_ARGUMENT(FTrailVisibilityManager*, InVisibilityManager)
+	SLATE_ARGUMENT(FGuid, InTrailGuid)
+	SLATE_ARGUMENT(FTrajectoryDrawInfo*, InDrawInfo)
+
+	SLATE_END_ARGS()
+
+
+public:
+	FTrailVisibilityManager* VisibilityManager;
+	FGuid TrailGuid;
+	FTrajectoryDrawInfo* DrawInfo;
+
+	void SetTrailAlwaysVisible(ECheckBoxState NewState)
+	{
+		bool bSet = (NewState == ECheckBoxState::Checked);
+		if (VisibilityManager)
+		{
+			VisibilityManager->SetTrailAlwaysVisible(TrailGuid,bSet);
+			if (DrawInfo && bSet == true)
+			{
+				DrawInfo->SetColor(UMotionTrailToolOptions::GetTrailOptions()->TrailColor);
+			}
+		}
+	}
+
+	ECheckBoxState IsTrailAlwaysVisible() const
+	{
+		if (VisibilityManager)
+		{
+			return VisibilityManager->IsTrailAlwaysVisible(TrailGuid) ? ECheckBoxState::Checked : ECheckBoxState::Unchecked;
+		}
+		return ECheckBoxState::Undetermined;
+	}
+
+	FReply ShowOptions()
+	{
+		const FName ControlRigMotionTrailTab("ControlRigMotionTrailTab");
+		FGlobalTabmanager::Get()->TryInvokeTab(ControlRigMotionTrailTab);
+		return FReply::Handled();
+	}
+
+	void Construct(const FArguments& InArgs)
+	{
+		VisibilityManager = InArgs._InVisibilityManager;
+		TrailGuid = InArgs._InTrailGuid;
+		DrawInfo = InArgs._InDrawInfo;
+		// Then make widget
+		this->ChildSlot
+		[
+			SNew(SBorder)
+			.BorderImage(FEditorStyle::GetBrush(TEXT("Menu.Background")))
+			.Padding(5)
+			.Content()
+			[
+				SNew(SVerticalBox)
+				+ SVerticalBox::Slot()
+				.AutoHeight()
+				.Padding(0.0f, 1.0f)
+				[
+					SNew(SHorizontalBox)
+					+ SHorizontalBox::Slot()
+					.AutoWidth()
+					.VAlign(VAlign_Center)
+					[
+						SNew(SCheckBox)
+						.OnCheckStateChanged(this, &STrailOptionsPopup::SetTrailAlwaysVisible)
+						.IsChecked(this, &STrailOptionsPopup::IsTrailAlwaysVisible)
+					]
+
+					+ SHorizontalBox::Slot()
+					.FillWidth(1.0f)
+					.VAlign(VAlign_Center)
+					[
+						SNew(STextBlock)
+						.AutoWrapText(true)
+						.Text(LOCTEXT("AlwaysVisible", "Always Visible"))
+					]	
+				]
+				+ SVerticalBox::Slot()
+				.AutoHeight()
+				.Padding(0.0f, 1.0f)
+				[
+					SNew(SButton)
+					.ContentPadding(FMargin(8, 2))
+					.Text(LOCTEXT("Show Options", "Show Options"))
+					.OnClicked(this, &STrailOptionsPopup::ShowOptions)
+				]
+			]
+		];
+		
+	}
+};
+
+void FTrailHierarchy::OpenContextMenu(const FGuid& TrailGuid)
+{
+
+	FLevelEditorModule& LevelEditorModule = FModuleManager::GetModuleChecked<FLevelEditorModule>("LevelEditor");
+	TSharedPtr< ILevelEditor > LevelEditor = LevelEditorModule.GetFirstLevelEditor();
+
+	TSharedPtr<SWidget> MenuWidget;
+	if(const TUniquePtr<FTrail>* Trail = GetAllTrails().Find(TrailGuid))
+	{
+		FTrajectoryDrawInfo* CurDrawInfo = (*Trail)->GetDrawInfo();
+
+		MenuWidget =
+			SNew(STrailOptionsPopup)
+			.InVisibilityManager(&GetVisibilityManager())
+			.InTrailGuid(TrailGuid)
+			.InDrawInfo(CurDrawInfo);
+
+		// Create as context menu
+		FSlateApplication::Get().PushMenu(
+			LevelEditor.ToSharedRef(),
+			FWidgetPath(),
+			MenuWidget.ToSharedRef(),
+			FSlateApplication::Get().GetCursorPos(),
+			FPopupTransitionEffect(FPopupTransitionEffect::ContextMenu)
+		);
+	}
+	
+}
+
+bool FTrailHierarchy::IsHitByClick(HHitProxy* InHitProxy)
+{
+	if (InHitProxy)
+	{
+		if (HBaseTrailProxy* HitProxy = HitProxyCast<HBaseTrailProxy>(InHitProxy))
+		{
+			return true;
+		}
+	}
+	return false;
+}
+
+bool FTrailHierarchy::HandleClick(FEditorViewportClient* InViewportClient, HHitProxy* InHitProxy,FInputClick Click)
+{
+	bool bClickIsHandled = false;
+	
+	if (HMotionTrailProxy* HitProxy = HitProxyCast<HMotionTrailProxy>(InHitProxy))
+	{
+		if (Click.bIsRightMouse)
+		{
+			OpenContextMenu(HitProxy->Guid);
+			return true;
+		}
+	}
+	
+
+	for (const TPair<FGuid, TUniquePtr<FTrail>>& GuidTrailPair :GetAllTrails())
+	{
+		bool bIsClickHandled = GuidTrailPair.Value->HandleClick(GuidTrailPair.Key, InViewportClient, InHitProxy, Click);
+		if (bIsClickHandled)
+		{
+			bClickIsHandled = true;
+		}
+	}
+
+	return bClickIsHandled;
+}
+
+bool FTrailHierarchy::IsAnythingSelected() const
+{
+	for (const TPair<FGuid, TUniquePtr<FTrail>>& GuidTrailPair : GetAllTrails())
+	{
+		FVector Location;
+		bool bHandled = GuidTrailPair.Value->IsAnythingSelected();
+		if (bHandled)
+		{
+			return true;
+		}
+	}
+	return false;
+}
+
+bool FTrailHierarchy::IsAnythingSelected(FVector& OutVectorPosition) const
+{
+	OutVectorPosition = FVector::ZeroVector;
+	int NumSelected = 0;
+	for (const TPair<FGuid, TUniquePtr<FTrail>>& GuidTrailPair : GetAllTrails())
+	{
+		FVector Location;
+		bool bHandled = GuidTrailPair.Value->IsAnythingSelected(Location);
+		if (bHandled)
+		{
+			OutVectorPosition += Location;
+			++NumSelected;
+		}
+	}
+
+	if (NumSelected > 0)
+	{
+		OutVectorPosition /= (double)NumSelected;
+	}
+	return (NumSelected > 0);
+}
+
+bool FTrailHierarchy::IsAnythingSelected(TArray<FVector>& OutVectorPositions) const
+{
+	OutVectorPositions.SetNum(0);
+	for (const TPair<FGuid, TUniquePtr<FTrail>>& GuidTrailPair : GetAllTrails())
+	{
+		FVector Location;
+		bool bHandled = GuidTrailPair.Value->IsAnythingSelected(Location);
+		if (bHandled)
+		{
+			OutVectorPositions.Add(Location);
+		}
+	}
+	return (OutVectorPositions.Num() > 0);
+}
+
+
+bool FTrailHierarchy::IsSelected(const FGuid& Key) const
+{
+	const TUniquePtr<FTrail>* Trail = GetAllTrails().Find(Key);
+	if (Trail != nullptr)
+	{
+		return (*Trail)->IsAnythingSelected();
+	}
+	return false;
+}
+void FTrailHierarchy::SelectNone()
+{
+	for (const TPair<FGuid, TUniquePtr<FTrail>>& GuidTrailPair : GetAllTrails())
+	{
+		GuidTrailPair.Value->SelectNone();
+	}
+}
+void FTrailHierarchy::TranslateSelectedKeys(bool bRight)
+{
+	for (const TPair<FGuid, TUniquePtr<FTrail>>& GuidTrailPair : GetAllTrails())
+	{		
+		GuidTrailPair.Value->TranslateSelectedKeys(bRight);
+	}
+}
+
+void FTrailHierarchy::DeleteSelectedKeys()
+{
+	for (const TPair<FGuid, TUniquePtr<FTrail>>& GuidTrailPair : GetAllTrails())
+	{
+		GuidTrailPair.Value->DeleteSelectedKeys();
+	}
+}
+
+bool FTrailHierarchy::IsAlwaysVisible(const FGuid Key) const
+{
+	return VisibilityManager.IsTrailAlwaysVisible(Key);
+}
+
+bool FTrailHierarchy::BoxSelect(FBox& InBox, bool InSelect)
+{
+	bool bIsHandled = false;
+	for (const TPair<FGuid, TUniquePtr<FTrail>>& GuidTrailPair : GetAllTrails())
+	{
+		bool bHandled = GuidTrailPair.Value->BoxSelect(InBox,InSelect);
+		if (bHandled)
+		{
+			bIsHandled = true;
+		}
+	}
+	return bIsHandled;
+}
+bool FTrailHierarchy::FrustumSelect(const FConvexVolume& InFrustum, FEditorViewportClient* InViewportClient, bool InSelect)
+{
+	bool bIsHandled = false;
+	for (const TPair<FGuid, TUniquePtr<FTrail>>& GuidTrailPair : GetAllTrails())
+	{
+		bool bHandled = GuidTrailPair.Value->FrustumSelect(InFrustum,InViewportClient,InSelect);
+		if (bHandled)
+		{
+			bIsHandled = true;
+		}
+	}
+	return bIsHandled;
+}
+
+
+bool FTrailHierarchy::StartTracking()
+{
+	bool bIsHandled = false;
+	for (const TPair<FGuid, TUniquePtr<FTrail>>& GuidTrailPair : GetAllTrails())
+	{
+		bool bHandled = GuidTrailPair.Value->StartTracking();
+		if (bHandled)
+		{
+			bIsHandled = true;
+		}
+	}
+	return bIsHandled;
+}
+
+bool FTrailHierarchy::ApplyDelta(const FVector& Pos, const FRotator& Rot, const FVector &WidgetLocation)
+
+{
+	bool bIsHandled = false;
+	for (const TPair<FGuid, TUniquePtr<FTrail>>& GuidTrailPair : GetAllTrails())
+	{
+		bool bHandled =  GuidTrailPair.Value->ApplyDelta(Pos,Rot,WidgetLocation);
+		if (bHandled)
+		{
+			bIsHandled = true;
+		}
+	}
+	return bIsHandled;
+}
+
+bool FTrailHierarchy::EndTracking()
+{
+	bool bIsHandled = false;
+	for (const TPair<FGuid, TUniquePtr<FTrail>>& GuidTrailPair : GetAllTrails())
+	{
+		bool bHandled = GuidTrailPair.Value->EndTracking();
+		if (bHandled)
+		{
+			bIsHandled = true;
+		}
+	}
+	return bIsHandled;
+}
+
+void FTrailHierarchy::CalculateEvalRangeArray()
+{
+	// Generate times to evaluate
+	SecondsPerSegment = GetSecondsPerSegment();
+
+	EvalTimesArr.Reserve(int32(EvalRange.Size<double>() / SecondsPerSegment) + 1);
+	EvalTimesArr.SetNum(0);
+	for (double SecondsItr = EvalRange.GetLowerBoundValue(); SecondsItr <= EvalRange.GetUpperBoundValue() + SecondsPerSegment; SecondsItr += SecondsPerSegment)
+	{
+		EvalTimesArr.Add(SecondsItr);
+	}
+
+	if (LastSecondsPerSegment != SecondsPerSegment || EvalRange != LastEvalRange)
+	{
+		for (const TPair<FGuid, TUniquePtr<FTrail>>& GuidTrailPair : GetAllTrails())
+		{
+			GuidTrailPair.Value->ForceEvaluateNextTick();
+		}
+		LastSecondsPerSegment = SecondsPerSegment;
+		LastEvalRange = EvalRange;
+	}
+}
+void FTrailHierarchy::Update()
+{
+	const FDateTime UpdateStartTime = FDateTime::Now();
+	
+	CalculateEvalRangeArray();
+	FTrailEvaluateTimes EvalTimes = FTrailEvaluateTimes(EvalTimesArr, SecondsPerSegment);
+
+	VisibilityManager.InactiveMask.Reset();
+	TArray<FGuid> DeadTrails;
+
+	for (const TPair<FGuid, TUniquePtr<FTrail>>& GuidTrailPair : GetAllTrails())
+	{
+		FGuid CurGuid = GuidTrailPair.Key;
+
+		FTrail::FSceneContext SceneContext = {
+			CurGuid,
+			EvalTimes,
+			this
+		};
+
+		if (!AllTrails.Contains(CurGuid))
+		{
+			continue;
+		}
+
+		// Update the trail
+		ETrailCacheState CurCacheState = AllTrails[CurGuid]->UpdateTrail(SceneContext);
+		if (CurCacheState == ETrailCacheState::Dead)
+		{
+			DeadTrails.Add(CurGuid);
+		}
+
+		if (CurCacheState == ETrailCacheState::NotUpdated)
+		{
+			VisibilityManager.InactiveMask.Add(CurGuid);
+		}
+	}
+
+	for (const FGuid& TrailGuid : DeadTrails)
+	{
+		RemoveTrail(TrailGuid);
+	}
+
+	const FTimespan UpdateTimespan = FDateTime::Now() - UpdateStartTime;
+	TimingStats.Add("FTrailHierarchy::Update", UpdateTimespan);
+}
+
+void FTrailHierarchy::AddTrail(const FGuid& Key, TUniquePtr<FTrail>&& TrailPtr)
+{
+	AllTrails.Add(Key, MoveTemp(TrailPtr));
+}
+
+void FTrailHierarchy::RemoveTrail(const FGuid& Key)
+{
+	AllTrails.Remove(Key);
+}
+
+void FTrailHierarchy::RemoveTrailIfNotAlwaysVisible(const FGuid& Key)
+{
+	if (IsAlwaysVisible(Key) == false)
+	{
+		RemoveTrail(Key);
+	}
+}
+
+} // namespace SequencerAnimTools
+} // namespace UE
+
+#undef LOCTEXT_NAMESPACE
diff --git a/Engine/Plugins/Editor/SequencerAnimTools/Source/SequencerAnimTools/Private/TrajectoryCache.cpp b/Engine/Plugins/Editor/SequencerAnimTools/Source/SequencerAnimTools/Private/TrajectoryCache.cpp
new file mode 100644
--- /dev/null
+++ b/Engine/Plugins/Editor/SequencerAnimTools/Source/SequencerAnimTools/Private/TrajectoryCache.cpp
@@ -0,0 +1,103 @@
+// Copyright Epic Games, Inc. All Rights Reserved.
+
+#include "TrajectoryCache.h"
+
+namespace UE
+{
+namespace SequencerAnimTools
+{
+
+bool RangesContain(const TSet<TRange<double>>& Ranges, const double Time)
+{
+	for (const TRange<double>& Range : Ranges)
+	{
+		if (Range.Contains(Time))
+		{
+			return true;
+		}
+	}
+	return false;
+}
+
+void FTrajectoryCache::UpdateCacheTimes(const FTrailEvaluateTimes& EvalTimes)
+{
+	TrackRange = EvalTimes.Range;
+	Spacing = EvalTimes.Spacing;
+}
+
+FTransform FArrayTrajectoryCache::GetInterp(const double InTime) const
+{
+	if (TrajectoryCache.Num() == 0)
+	{
+		return Default;
+	}
+
+	const double T = (InTime/Spacing) - FMath::FloorToDouble(InTime / Spacing);
+	const int32 LowIdx = FMath::Clamp(int32((InTime - TrackRange.GetLowerBoundValue()) / Spacing), 0, TrajectoryCache.Num() - 1);
+	const int32 HighIdx = FMath::Clamp(LowIdx + 1, 0, TrajectoryCache.Num() - 1);
+
+	if (LowIdx == HighIdx)
+	{
+		return TrajectoryCache[LowIdx];
+	}
+
+	FTransform TempBlended; 
+	TempBlended.Blend(TrajectoryCache[LowIdx], TrajectoryCache[HighIdx], T);
+	return TempBlended;
+}
+
+void FArrayTrajectoryCache::GetInterp(const double InTime,FTransform& Transform, FTransform& ParentTransform) const
+{
+	if (TrajectoryCache.Num() == 0)
+	{
+		Transform =ParentTransform = Default;
+	}
+
+	const double T = (InTime / Spacing) - FMath::FloorToDouble(InTime / Spacing);
+	const int32 LowIdx = FMath::Clamp(int32((InTime - TrackRange.GetLowerBoundValue()) / Spacing), 0, TrajectoryCache.Num() - 1);
+	const int32 HighIdx = FMath::Clamp(LowIdx + 1, 0, TrajectoryCache.Num() - 1);
+
+	if (LowIdx == HighIdx)
+	{
+		Transform = TrajectoryCache[LowIdx];
+		if (ParentTrajectoryCache.Num() > 0)
+		{
+			ParentTransform = ParentTrajectoryCache[LowIdx];
+		}
+		else
+		{
+			ParentTransform = FTransform::Identity;
+		}
+	}
+
+	Transform.Blend(TrajectoryCache[LowIdx], TrajectoryCache[HighIdx], T);
+	if (ParentTrajectoryCache.Num() > 0)
+	{
+		ParentTransform.Blend(ParentTrajectoryCache[LowIdx], ParentTrajectoryCache[HighIdx], T);
+	}
+	else
+	{
+		ParentTransform = FTransform::Identity;
+	}
+
+}
+
+TArray<double> FArrayTrajectoryCache::GetAllTimesInRange(const TRange<double>& InRange) const
+{
+	TRange<double> GenRange = TRange<double>::Intersection({ TrackRange, InRange });
+
+	TArray<double> AllTimesInRange;
+	AllTimesInRange.Reserve(int(GenRange.Size<double>() / Spacing) + 1);
+	const double FirstTick = FMath::FloorToDouble((GenRange.GetLowerBoundValue() / Spacing) + KINDA_SMALL_NUMBER) * Spacing;
+	for (double TickItr = FirstTick + KINDA_SMALL_NUMBER; TickItr < GenRange.GetUpperBoundValue(); TickItr += Spacing)
+	{
+		AllTimesInRange.Add(TickItr);
+	}
+
+	AllTimesInRange.Add(GenRange.GetUpperBoundValue());
+
+	return AllTimesInRange;
+}
+
+} // namespace SequencerAnimTools
+} // namespace UE
diff --git a/Engine/Plugins/Editor/SequencerAnimTools/Source/SequencerAnimTools/Private/TrajectoryDrawInfo.cpp b/Engine/Plugins/Editor/SequencerAnimTools/Source/SequencerAnimTools/Private/TrajectoryDrawInfo.cpp
new file mode 100644
--- /dev/null
+++ b/Engine/Plugins/Editor/SequencerAnimTools/Source/SequencerAnimTools/Private/TrajectoryDrawInfo.cpp
@@ -0,0 +1,64 @@
+// Copyright Epic Games, Inc. All Rights Reserved.
+
+#include "TrajectoryDrawInfo.h"
+#include "TrailHierarchy.h"
+#include "TrajectoryCache.h"
+
+namespace UE
+{
+namespace SequencerAnimTools
+{
+
+TOptional<FVector2D> FTrailScreenSpaceTransform::ProjectPoint(const FVector& Point) const
+{
+	const FPlane Projection = View->Project(Point);
+	if (Projection.W > 0.0f)
+	{
+		const float XPos = HalfScreenSize.X + (HalfScreenSize.X * Projection.X * 1.0f);
+		const float YPos = HalfScreenSize.Y + (HalfScreenSize.Y * Projection.Y * -1.0f);
+		return FVector2D(XPos, YPos);
+	}
+
+	return TOptional<FVector2D>();
+}
+
+void FTrajectoryDrawInfo::GetTrajectoryPointsForDisplay(const FDisplayContext& InDisplayContext, TArray<FVector>& TrajectoryPoints, TArray<double>& TrajectoryTimes)
+{	
+	CachedViewRange = InDisplayContext.TimeRange;
+
+	TrajectoryTimes = TrajectoryCache->GetAllTimesInRange(InDisplayContext.TimeRange);
+	TrajectoryPoints.Reserve(TrajectoryTimes.Num());
+	
+	for (const double Time : TrajectoryTimes)
+	{
+		TrajectoryPoints.Add(TrajectoryCache->Get(Time).GetTranslation());
+	}
+
+}
+
+void FTrajectoryDrawInfo::GetTickPointsForDisplay(const FDisplayContext& InDisplayContext, TArray<FVector2D>& Ticks, TArray<FVector2D>& TickNormals)
+{
+	const double FirstTick = FMath::FloorToDouble(InDisplayContext.TimeRange.GetLowerBoundValue() / InDisplayContext.SecondsPerTick) * InDisplayContext.SecondsPerTick;
+	const double Spacing = InDisplayContext.SecondsPerSegment;
+
+	for (double TickItr = FirstTick + InDisplayContext.SecondsPerTick; TickItr < InDisplayContext.TimeRange.GetUpperBoundValue(); TickItr += InDisplayContext.SecondsPerTick)
+	{
+		FVector Interpolated = TrajectoryCache->GetInterp(TickItr).GetTranslation();
+		TOptional<FVector2D> Projected = InDisplayContext.ScreenSpaceTransform.ProjectPoint(Interpolated);
+		const double PrevSampleTime = TickItr + Spacing < InDisplayContext.TimeRange.GetUpperBoundValue() ? TickItr + Spacing : TickItr - Spacing;
+		TOptional<FVector2D> PrevProjected = InDisplayContext.ScreenSpaceTransform.ProjectPoint(TrajectoryCache->GetInterp(PrevSampleTime).GetTranslation());
+
+		if (Projected && PrevProjected)
+		{
+			Ticks.Add(Projected.GetValue());
+
+			FVector2D Diff = Projected.GetValue() - PrevProjected.GetValue();
+			Diff.Normalize();
+
+			TickNormals.Add(FVector2D(-Diff.Y, Diff.X));
+		}
+	}
+}
+
+} // namespace SequencerAnimTools
+} // namespace UE
diff --git a/Engine/Plugins/Editor/SequencerAnimTools/Source/SequencerAnimTools/Public/BaseSequencerAnimTool.h b/Engine/Plugins/Editor/SequencerAnimTools/Source/SequencerAnimTools/Public/BaseSequencerAnimTool.h
new file mode 100644
--- /dev/null
+++ b/Engine/Plugins/Editor/SequencerAnimTools/Source/SequencerAnimTools/Public/BaseSequencerAnimTool.h
@@ -0,0 +1,21 @@
+// Copyright Epic Games, Inc. All Rights Reserved.
+
+#pragma once
+
+#include "CoreMinimal.h"
+#include "UObject/ObjectMacros.h"
+#include "UObject/Interface.h"
+#include "BaseSequencerAnimTool.generated.h"
+
+UINTERFACE(MinimalAPI)
+class UBaseSequencerAnimTool : public UInterface
+{
+	GENERATED_UINTERFACE_BODY()
+};
+
+class SEQUENCERANIMTOOLS_API IBaseSequencerAnimTool
+{
+	GENERATED_IINTERFACE_BODY()
+	virtual bool ProcessCommandBindings(const FKey Key, const bool bRepeat) const { return false; }
+
+};
diff --git a/Engine/Plugins/Editor/SequencerAnimTools/Source/SequencerAnimTools/Public/EditPivotTool.h b/Engine/Plugins/Editor/SequencerAnimTools/Source/SequencerAnimTools/Public/EditPivotTool.h
new file mode 100644
--- /dev/null
+++ b/Engine/Plugins/Editor/SequencerAnimTools/Source/SequencerAnimTools/Public/EditPivotTool.h
@@ -0,0 +1,172 @@
+// Copyright Epic Games, Inc. All Rights Reserved.
+
+#pragma once
+
+#include "CoreMinimal.h"
+#include "UObject/Object.h"
+
+#include "InteractiveTool.h"
+#include "InteractiveToolBuilder.h"
+#include "MultiSelectionTool.h"
+#include "BaseBehaviors/BehaviorTargetInterfaces.h"
+#include "Framework/Commands/UICommandInfo.h"
+#include "BaseGizmos/TransformProxy.h"
+#include "BaseSequencerAnimTool.h"
+#include "EditPivotTool.generated.h"
+
+class USingleClickInputBehavior;
+class UClickDragInputBehavior;
+class UTransformGizmo;
+class ULevelSequence;
+class UControlRig;
+struct FRigControlElement;
+class ISequencer;
+
+/*
+*  The way this sequencer pivot tool works is that
+*  it will use modify the incoming selections temp pivot
+*  while the mode is active. Reselecting will turn off the mode.
+*
+*/
+
+/**
+ * Builder for USequencerPivotTool
+ */
+UCLASS()
+class  USequencerPivotToolBuilder : public UInteractiveToolBuilder
+{
+	GENERATED_BODY()
+
+public:
+	virtual bool CanBuildTool(const FToolBuilderState& SceneState) const override;
+	virtual UInteractiveTool* BuildTool(const FToolBuilderState& SceneState) const override;
+};
+
+struct FControlRigSelectionDuringDrag
+{
+	ULevelSequence* LevelSequence;
+	FFrameNumber CurrentFrame;
+	UControlRig* ControlRig;
+	FName ControlName;
+	FTransform CurrentTransform;
+};
+
+struct FControlRigMappings
+{
+	TWeakObjectPtr<UControlRig> ControlRig;
+	TMap<FName, FTransform> PivotTransforms;
+};
+
+class FEditPivotCommands : public TCommands<FEditPivotCommands>
+{
+public:
+	FEditPivotCommands()
+		: TCommands<FEditPivotCommands>(
+			"SequencerEditPivotTool",
+			NSLOCTEXT("SequencerEditPivotTool", "SequencerEditPivotTool", "Edit Pivot Commands"),
+			NAME_None,
+			FEditorStyle::GetStyleSetName()
+			)
+	{}
+
+
+	virtual void RegisterCommands() override;
+	static const TMap<FName, TArray<TSharedPtr<FUICommandInfo>>>& GetCommands()
+	{
+		return FEditPivotCommands::Get().Commands;
+	}
+
+public:
+	/** Reset the Pivot*/
+	TSharedPtr<FUICommandInfo> ResetPivot;
+
+	TMap<FName, TArray<TSharedPtr<FUICommandInfo>>> Commands;
+};
+
+/**
+
+ */
+UCLASS()
+class  USequencerPivotTool : public UMultiSelectionTool, public IClickBehaviorTarget  , public IBaseSequencerAnimTool
+{
+	GENERATED_BODY()
+
+public:
+
+	// UInteractiveTool overrides
+	virtual void SetWorld(UWorld* World, UInteractiveGizmoManager* GizmoManager);
+	virtual void Setup() override;
+	virtual void Render(IToolsContextRenderAPI* RenderAPI) override;
+	virtual void OnPropertyModified(UObject* PropertySet, FProperty* Property) override;
+	virtual void Shutdown(EToolShutdownType ShutdownType) override;
+	virtual bool HasCancel() const override { return false; }
+	virtual bool HasAccept() const override { return false; }
+	virtual bool CanAccept() const override { return false; }
+
+	// IClickBehaviorTarget interface
+	virtual FInputRayHit IsHitByClick(const FInputDeviceRay& ClickPos) override;
+	virtual void OnClicked(const FInputDeviceRay& ClickPos) override;
+
+	//IBaseSequencerAnimTool
+	virtual bool ProcessCommandBindings(const FKey Key, const bool bRepeat) const override;
+
+	// End interfaces
+
+protected:
+
+	UPROPERTY()
+	USingleClickInputBehavior* ClickBehavior = nullptr;
+
+	UPROPERTY()
+	UTransformProxy* TransformProxy = nullptr;
+
+	UPROPERTY()
+	UTransformGizmo* TransformGizmo = nullptr;
+
+protected:
+	bool bShiftPressedWhenStarted = false;
+	int32 CtrlModifierId = 1;
+	UWorld* TargetWorld = nullptr;		// target World we will raycast into
+	UInteractiveGizmoManager* GizmoManager = nullptr; //gizmo man
+
+	FTransform StartDragTransform;
+	bool bGizmoBeingDragged = false;
+	bool bManipulatorMadeChange = false;
+	int32 TransactionIndex = -1;
+	TArray<FControlRigSelectionDuringDrag> ControlRigDrags;
+
+	//since we are selection based we can cache this
+	ULevelSequence* LevelSequence;
+	TArray<UControlRig*> ControlRigs;
+	TWeakPtr<ISequencer> SequencerPtr;
+
+	FTransform GizmoTransform = FTransform::Identity;
+	bool bPickingPivotLocation = false; //mauy remove
+	void UpdateGizmoVisibility();
+	void UpdateGizmoTransform();
+	void UpdateTransformAndSelectionOnEntering();
+	bool SetGizmoBasedOnSelection(bool bUseSaved = true);
+
+	FInputRayHit FindRayHit(const FRay& WorldRay, FVector& HitPos);		// raycasts into World
+
+	// Callbacks we'll receive from the gizmo proxy
+	void GizmoTransformChanged(UTransformProxy* Proxy, FTransform Transform);
+	void GizmoTransformStarted(UTransformProxy* Proxy);
+	void GizmoTransformEnded(UTransformProxy* Proxy);
+
+	//Handle Selection and Pivot Location
+	void SavePivotTransforms();
+	void SaveLastSelected();
+	void RemoveControlRigDelegates();
+	void HandleControlSelected(UControlRig* Subject, FRigControlElement* InControl, bool bSelected);
+
+private:
+	TSharedPtr<FUICommandList> CommandBindings;
+	void ResetPivot();
+
+public:
+	static TMap<TWeakObjectPtr<UControlRig>, FControlRigMappings> SavedPivotLocations;
+	static TArray<FControlRigMappings> LastSelectedObjects;
+};
+
+
diff --git a/Engine/Plugins/Editor/SequencerAnimTools/Source/SequencerAnimTools/Public/MotionTrailTool.h b/Engine/Plugins/Editor/SequencerAnimTools/Source/SequencerAnimTools/Public/MotionTrailTool.h
new file mode 100644
--- /dev/null
+++ b/Engine/Plugins/Editor/SequencerAnimTools/Source/SequencerAnimTools/Public/MotionTrailTool.h
@@ -0,0 +1,232 @@
+// Copyright Epic Games, Inc. All Rights Reserved.
+
+#pragma once
+
+#include "CoreMinimal.h"
+#include "UObject/Object.h"
+
+#include "InteractiveTool.h"
+#include "InteractiveToolBuilder.h"
+#include "MultiSelectionTool.h"
+#include "BaseBehaviors/BehaviorTargetInterfaces.h"
+#include "Framework/Commands/UICommandInfo.h"
+#include "BaseGizmos/TransformProxy.h"
+#include "Framework/Commands/Commands.h"
+#include "ISequencer.h"
+#include "TrailHierarchy.h"
+#include "Framework/Commands/Commands.h"
+#include "EditorStyleSet.h"
+#include "BaseSequencerAnimTool.h"
+#include "MotionTrailTool.generated.h"
+
+class USingleClickInputBehavior;
+class UClickDragInputBehavior;
+class UTransformGizmo;
+class FEditorModeTools;
+class USingleKeyCaptureBehavior;
+
+
+/** 
+* IClickBehaviorTarget Wrapper to handle left and right clicks
+*/
+class FLeftRghtClickBehaviorTarget: public IClickBehaviorTarget
+{
+
+
+public:
+	FLeftRghtClickBehaviorTarget() {};
+
+	IClickBehaviorTarget* Target = nullptr;
+	bool bIsLeft = true;
+	//Target will check these and then revert them back after using them.
+	bool bClicked = false;
+
+	virtual FInputRayHit IsHitByClick(const FInputDeviceRay& ClickPos) override
+	{
+		bClicked = true;
+		FInputRayHit Hit = Target->IsHitByClick(ClickPos);
+		bClicked = false;
+		return Hit;
+	}
+
+	virtual void OnClicked(const FInputDeviceRay& ClickPos) override
+	{
+		bClicked = true;
+		Target->OnClicked(ClickPos);
+		bClicked = false;
+	}
+
+	virtual void OnUpdateModifierState(int ModifierID, bool bIsOn) override
+	{
+		bClicked = true;
+		Target->OnUpdateModifierState(ModifierID, bIsOn);
+		bClicked = false;
+	}
+
+};
+
+
+/**
+ * Builder for UMotionTrailTool
+ */
+UCLASS()
+class UMotionTrailToolBuilder : public UInteractiveToolBuilder
+{
+	GENERATED_BODY()
+
+public:
+	virtual bool CanBuildTool(const FToolBuilderState& SceneState) const override { return true; }
+	virtual UInteractiveTool* BuildTool(const FToolBuilderState& SceneState) const override;
+
+};
+
+class FMotionTrailCommands : public TCommands<FMotionTrailCommands>
+{
+public:
+	FMotionTrailCommands()
+		: TCommands<FMotionTrailCommands>(
+			"MotionTrailCommands",
+			NSLOCTEXT("MotionTrailTool", "MotionTrailCommands", "Motion Trail Commands"),
+			NAME_None,
+			FEditorStyle::GetStyleSetName()
+		)
+	{}
+
+
+	virtual void RegisterCommands() override;
+	static const TMap<FName, TArray<TSharedPtr<FUICommandInfo>>>& GetCommands()
+	{
+		return FMotionTrailCommands::Get().Commands;
+	}
+
+public:
+	/** Move selected keys one frame left*/
+	TSharedPtr<FUICommandInfo> TranslateSelectedKeysLeft;
+
+	/** Move seleced keys one frame right*/
+	TSharedPtr<FUICommandInfo> TranslateSelectedKeysRight;
+
+	/** Frame the selection*/
+	TSharedPtr<FUICommandInfo> FrameSelection;
+
+	TMap<FName, TArray<TSharedPtr<FUICommandInfo>>> Commands;
+};
+
+
+UCLASS()
+class UMotionTrailTool : public UMultiSelectionTool, public IClickBehaviorTarget, public IBaseSequencerAnimTool
+{
+	GENERATED_BODY()
+
+public:
+
+	// IClickBehaviorTarget interface
+	virtual FInputRayHit IsHitByClick(const FInputDeviceRay& ClickPos);
+	virtual void OnClicked(const FInputDeviceRay& ClickPos);
+
+	// IModifierToggleBehaviorTarget implementation, inherited through IClickBehaviorTarget
+	virtual void OnUpdateModifierState(int ModifierID, bool bIsOn) override;
+
+	// UInteractiveTool overrides
+	virtual	void Setup() override;
+	virtual void Shutdown(EToolShutdownType ShutdownType) override;
+	virtual void Render(IToolsContextRenderAPI* RenderAPI) override;
+	virtual void DrawHUD(FCanvas* Canvas, IToolsContextRenderAPI* RenderAPI) override;
+
+	virtual void OnTick(float DeltaTime) override;
+
+	virtual TArray<UObject*> GetToolProperties(bool bEnabledOnly = true) const override;
+	
+	//IBaseSequencerAnimTool
+	virtual bool ProcessCommandBindings(const FKey Key, const bool bRepeat) const override;
+
+	// End interfaces
+
+	void SetWorldGizmoModeManager(UWorld* World, UInteractiveGizmoManager* InGizmoManager, FEditorModeTools* InModeManager) 
+	{
+		TargetWorld = World;
+		GizmoManager = InGizmoManager; 
+		ModeManager = InModeManager;
+	}
+	UWorld* GetWorld() const { return TargetWorld; }
+
+	FEditorModeTools* GetModeManager() const { return ModeManager;}
+	static FString TrailKeyTransformGizmoInstanceIdentifier;
+public:
+
+	// Support for gizmo. Since the points aren't individual components, we don't actually use UTransformProxy
+	// for the transform forwarding- we just use it for the callbacks.
+	UPROPERTY()
+	UTransformProxy* TransformProxy;
+	UPROPERTY()
+	UTransformGizmo* TransformGizmo;
+
+	FTransform StartDragTransform;
+	bool bGizmoBeingDragged = false;
+
+	FTransform GizmoTransform = FTransform::Identity;
+	
+	UE::SequencerAnimTools::FTrailHierarchy* GetHierarchyForSequencer(ISequencer* Sequencer) const { return SequencerHierarchies[Sequencer]; }
+	const TArray<TUniquePtr<UE::SequencerAnimTools::FTrailHierarchy>>& GetHierarchies() const { return TrailHierarchies; }
+
+
+	void SelectNone();
+protected:
+
+	void SetupIntegration();
+	void ShutdownIntegration();
+	void UpdateGizmoVisibility();
+	void UpdateGizmoLocation();
+
+	void GizmoTransformChanged(UTransformProxy* Proxy, FTransform Transform);
+	void GizmoTransformStarted(UTransformProxy* Proxy);
+	void GizmoTransformEnded(UTransformProxy* Proxy);
+
+protected:
+	//Actions
+	void TranslateSelectedKeysLeft();
+	void TranslateSelectedKeysRight();
+	void FrameSelection();
+	void DeleteSelectedKeys();
+	bool SomeKeysAreSelected() const;
+
+protected:
+	TArray<TUniquePtr<UE::SequencerAnimTools::FTrailHierarchy>> TrailHierarchies;
+	TMap<ISequencer*, UE::SequencerAnimTools::FTrailHierarchy*> SequencerHierarchies;
+
+	FDelegateHandle OnSequencersChangedHandle;
+	//ed mode for manips
+	bool bManipulatorMadeChange = false;
+	int32 TransactionIndex = -1;
+
+protected:
+
+	FLeftRghtClickBehaviorTarget LeftTarget;
+	FLeftRghtClickBehaviorTarget RightTarget;
+
+	UPROPERTY()
+	USingleClickInputBehavior* LeftClickBehavior = nullptr;
+
+	UPROPERTY()
+	USingleClickInputBehavior* RightClickBehavior = nullptr;
+
+
+
+	UWorld* TargetWorld = nullptr;
+	UInteractiveGizmoManager* GizmoManager = nullptr;
+	FEditorModeTools* ModeManager = nullptr;
+	// Support for Alt,Shift and Ctrl toggles
+	int32 ShiftModifierId = 1;
+	bool bShiftModifierOn = false;
+	int32 CtrlModifierId = 2;
+	bool bCtrlModifierOn = false;
+	// note we are using shift+ctrl for alt
+	int32 AltModiferId = 3;
+	bool bAltModifierOn = false;
+
+private:
+	TSharedPtr<FUICommandList> CommandBindings;
+
+
+};
+
diff --git a/Engine/Plugins/Editor/SequencerAnimTools/Source/SequencerAnimTools/Public/SequencerAnimToolsModule.h b/Engine/Plugins/Editor/SequencerAnimTools/Source/SequencerAnimTools/Public/SequencerAnimToolsModule.h
new file mode 100644
--- /dev/null
+++ b/Engine/Plugins/Editor/SequencerAnimTools/Source/SequencerAnimTools/Public/SequencerAnimToolsModule.h
@@ -0,0 +1,37 @@
+// Copyright Epic Games, Inc. All Rights Reserved.
+
+#pragma once
+
+#include "CoreMinimal.h"
+#include "Modules/ModuleManager.h"
+
+class ILevelEditor;
+
+namespace UE
+{
+namespace SequencerAnimTools
+{
+
+
+class FSequencerAnimToolsModule : public IModuleInterface
+{
+public:
+
+	/** IModuleInterface implementation */
+	virtual void StartupModule() override;
+	virtual void ShutdownModule() override;
+
+private:
+	TWeakPtr<ILevelEditor> LevelEditorPtr;
+	TSharedPtr<FUICommandList> CommandBindings;
+
+	//delegates
+	void OnLevelEditorCreated(TSharedPtr<ILevelEditor> InLevelEditor);
+	void OnMotionTralOptionChanged(FName PropretyName);
+
+	//need to keep track of which levels we registered the tools for otherwise we can hit an ensure
+	TSet<ILevelEditor*>  AlreadyRegisteredTools;
+};
+
+} // namespace SequencerAnimTools
+} // namespace UE
diff --git a/Engine/Plugins/Editor/SequencerAnimTools/Source/SequencerAnimTools/Public/SequencerToolsEditMode.h b/Engine/Plugins/Editor/SequencerAnimTools/Source/SequencerAnimTools/Public/SequencerToolsEditMode.h
new file mode 100644
--- /dev/null
+++ b/Engine/Plugins/Editor/SequencerAnimTools/Source/SequencerAnimTools/Public/SequencerToolsEditMode.h
@@ -0,0 +1,28 @@
+// Copyright Epic Games, Inc. All Rights Reserved.
+
+#pragma once
+
+#include "CoreMinimal.h"
+#include "Tools/LegacyEdModeWidgetHelpers.h"
+
+#include "SequencerToolsEditMode.generated.h"
+
+
+UCLASS()
+class USequencerToolsEditMode : public UBaseLegacyWidgetEdMode
+{
+	GENERATED_BODY()
+public:
+	static FEditorModeID ModeName;
+
+	USequencerToolsEditMode();
+	virtual ~USequencerToolsEditMode();
+
+	// UBaseLegacyWidgetEdMode interface
+	
+	virtual bool UsesToolkits() const  override { return false; }
+	virtual bool InputKey(FEditorViewportClient* InViewportClient, FViewport* InViewport, FKey InKey, EInputEvent InEvent) override;
+
+	virtual bool IsCompatibleWith(FEditorModeID OtherModeID) const override;
+
+};
diff --git a/Engine/Plugins/Editor/SequencerAnimTools/Source/SequencerAnimTools/Public/Trail.h b/Engine/Plugins/Editor/SequencerAnimTools/Source/SequencerAnimTools/Public/Trail.h
new file mode 100644
--- /dev/null
+++ b/Engine/Plugins/Editor/SequencerAnimTools/Source/SequencerAnimTools/Public/Trail.h
@@ -0,0 +1,130 @@
+// Copyright Epic Games, Inc. All Rights Reserved.
+
+#pragma once
+
+#include "Misc/Guid.h"
+#include "Containers/Map.h"
+#include "Containers/Array.h"
+#include "Math/Range.h"
+#include "Math/Color.h"
+#include "UObject/WeakObjectPtrTemplates.h"
+#include "GameFramework/Actor.h"
+
+#include "TrajectoryCache.h"
+#include "TrajectoryDrawInfo.h"
+#include "HitProxies.h"
+#include "Misc/Guid.h"
+
+
+
+class UMotionTrailEditorMode;
+namespace UE
+{
+namespace SequencerAnimTools
+{
+struct FInputClick
+{
+	FInputClick() {};
+	FInputClick(bool bA, bool bC, bool bS) : bAltIsDown(bA), bCtrlIsDown(bC), bShiftIsDown(bS) {};
+	bool bAltIsDown = false; 
+	bool bCtrlIsDown = false;
+	bool bShiftIsDown = false;
+	bool bIsRightMouse = false;
+};
+struct HBaseTrailProxy : public HHitProxy
+{
+	DECLARE_HIT_PROXY();
+
+	FGuid Guid;
+	HBaseTrailProxy(const FGuid& InGuid, EHitProxyPriority InPriority = HPP_Foreground) 
+		:HHitProxy(InPriority)
+		,Guid(InGuid)
+	{}
+};
+struct HMotionTrailProxy : public HBaseTrailProxy
+{
+	DECLARE_HIT_PROXY();
+
+	FVector Point;
+	double Second;
+
+	HMotionTrailProxy(const FGuid& InGuid, const FVector& InPoint, double InSecond) 
+		:HBaseTrailProxy(InGuid)
+		,Point(InPoint)
+		,Second(InSecond)
+	{}
+
+};
+
+enum class ETrailCacheState : uint8
+{
+	UpToDate = 2,
+	Stale = 1,
+	Dead = 0,
+	NotUpdated = 3
+};
+
+class FTrail
+{
+public:
+
+	FTrail(UObject* InOwner)
+		: Owner(InOwner)
+		,CacheState(ETrailCacheState::Stale)
+		, bForceEvaluateNextTick(true)
+		, DrawInfo(nullptr)
+	{}
+
+	virtual ~FTrail() {}
+
+	// Public interface methods
+	struct FSceneContext
+	{
+		FGuid YourNode;
+		const class FTrailEvaluateTimes& EvalTimes;
+		class FTrailHierarchy* TrailHierarchy;
+	};
+
+	UObject* GetOwner() const
+	{ 
+		return (Owner.IsValid() ? Owner.Get() : nullptr);
+	}
+
+	virtual ETrailCacheState UpdateTrail(const FSceneContext& InSceneContext) = 0;
+	virtual FTrajectoryCache* GetTrajectoryTransforms() = 0;
+
+	//Additional Render/HitTest event handling for specific trails, usually let default renderer handle it
+	virtual void Render(const FGuid& Guid, const FSceneView* View, FViewport* Viewport, FPrimitiveDrawInterface* PDI) {};
+	virtual void DrawHUD(FEditorViewportClient* ViewportClient, FViewport* Viewport, const FSceneView* View, FCanvas* Canvas) {};
+	virtual bool HandleClick(const FGuid& Guid, FEditorViewportClient* InViewportClient, HHitProxy* HitProxy, FInputClick Click) { return false; }
+	virtual bool IsAnythingSelected() const { return false; }
+	virtual bool IsAnythingSelected(FVector& OutVectorPosition)const { return false;}
+	virtual bool IsTrailSelected() const { return false; }
+	virtual bool BoxSelect(FBox& InBox, bool InSelect = true) {return false;}
+	virtual bool FrustumSelect(const FConvexVolume& InFrustum, FEditorViewportClient* InViewportClient, bool InSelect = true)  { return false; }
+
+	virtual bool StartTracking() { return false; }
+	virtual bool ApplyDelta(const FVector& Pos, const FRotator& Rot, const FVector& WidgetLocation) { return false; }
+	virtual bool EndTracking() { return false; }
+	virtual void TranslateSelectedKeys(bool bRight) {};
+	virtual void DeleteSelectedKeys() {};
+	virtual void SelectNone() {};
+	virtual bool GetEditedTimes(const FTrailHierarchy* TrailHierarchy, const FFrameNumber& LastFrame, TArray<FFrameNumber>& OutEditedTimes) { return false; }
+	virtual void UpdateKeysInRange(const TRange<double>& ViewRange) {};
+	// Optionally implemented methods
+	virtual TRange<double> GetEffectiveRange() const { return TRange<double>::Empty(); }
+
+	FTrajectoryDrawInfo* GetDrawInfo() { return DrawInfo.Get(); }
+	void ForceEvaluateNextTick() { bForceEvaluateNextTick = true; }
+
+	ETrailCacheState GetCacheState() const { return CacheState; }
+protected:
+
+	TWeakObjectPtr<UObject> Owner;
+	ETrailCacheState CacheState;
+
+	bool bForceEvaluateNextTick;
+	TUniquePtr<FTrajectoryDrawInfo> DrawInfo;
+};
+} // namespace MovieScene
+} // namespace UE
diff --git a/Engine/Plugins/Editor/SequencerAnimTools/Source/SequencerAnimTools/Public/TrailHierarchy.h b/Engine/Plugins/Editor/SequencerAnimTools/Source/SequencerAnimTools/Public/TrailHierarchy.h
new file mode 100644
--- /dev/null
+++ b/Engine/Plugins/Editor/SequencerAnimTools/Source/SequencerAnimTools/Public/TrailHierarchy.h
@@ -0,0 +1,156 @@
+// Copyright Epic Games, Inc. All Rights Reserved.
+
+#pragma once
+
+#include "Trail.h"
+#include "EditorViewportClient.h"
+#include "Widgets/Input/SCheckBox.h"
+#include "SceneView.h"
+#include "UnrealClient.h"
+#include "Misc/Timespan.h"
+
+
+namespace UE
+{
+namespace SequencerAnimTools
+{
+
+
+struct FTrailVisibilityManager
+{
+	bool IsTrailVisible(const FGuid& Guid, const FTrail* Trail) const
+	{
+		return !InactiveMask.Contains(Guid) && !VisibilityMask.Contains(Guid) && (AlwaysVisible.Contains(Guid) || Selected.Contains(Guid) || ControlSelected.Contains(Guid)
+			|| Trail->IsAnythingSelected()) && Guid.IsValid();
+	}
+
+	bool IsTrailAlwaysVisible(const FGuid& Guid) const
+	{
+		return AlwaysVisible.Contains(Guid);
+	}
+	void SetTrailAlwaysVisible(const FGuid& Guid, bool bSet)
+	{
+		if (bSet)
+		{
+			AlwaysVisible.Add(Guid);
+		}
+		else
+		{
+			AlwaysVisible.Remove(Guid);
+		}
+	}
+
+
+	TSet<FGuid> InactiveMask; // Any trails whose cache state or parent's cache state has been marked as NotUpdated
+	TSet<FGuid> VisibilityMask; // Any trails masked out by the user interface, ex bone trails
+	TSet<FGuid> AlwaysVisible; // Any trails pinned by the user interface
+	TSet<FGuid> Selected; // Any transform or bone trails selected in the user interface
+	TSet<FGuid> ControlSelected;// Any control rig trails selected
+};
+
+class ITrailHierarchyRenderer
+{
+public:
+	virtual ~ITrailHierarchyRenderer() {}
+	virtual void Render(const FSceneView* View, FViewport* Viewport, FPrimitiveDrawInterface* PDI) = 0;
+	virtual void DrawHUD(FEditorViewportClient* ViewportClient, FViewport* Viewport, const FSceneView* View, FCanvas* Canvas) = 0;
+};
+
+class FTrailHierarchyRenderer : public ITrailHierarchyRenderer
+{
+public:
+	FTrailHierarchyRenderer(class FTrailHierarchy* InOwningHierarchy)
+		: OwningHierarchy(InOwningHierarchy)
+	{}
+
+	virtual void Render(const FSceneView* View, FViewport* Viewport, FPrimitiveDrawInterface* PDI) override;
+	virtual void DrawHUD(FEditorViewportClient* ViewportClient, FViewport* Viewport, const FSceneView* View, FCanvas* Canvas) override;
+
+private:
+	class FTrailHierarchy* OwningHierarchy;
+};
+
+class FTrailHierarchy
+{
+public:
+
+	FTrailHierarchy()
+		: ViewRange(TRange<double>::All())
+		, EvalRange(TRange<double>::All())
+		, SecondsPerSegment(0.1)
+		, LastEvalRange(TRange<double>::All())
+		, LastSecondsPerSegment(0.1)
+		, AllTrails()
+		, TimingStats()
+		, VisibilityManager()
+	{}
+
+	virtual ~FTrailHierarchy() {}
+
+	virtual void Initialize() = 0;
+	virtual void Destroy() = 0; // TODO: make dtor?
+	virtual ITrailHierarchyRenderer* GetRenderer() const = 0;
+	virtual double GetSecondsPerFrame() const = 0;
+	virtual double GetSecondsPerSegment() const = 0;
+	virtual FFrameNumber GetFramesPerFrame() const = 0;
+	virtual FFrameNumber GetFramesPerSegment() const = 0;
+
+	// Optionally implemented methods
+	virtual void CalculateEvalRangeArray();
+	virtual void Update();
+	virtual bool HandleClick(FEditorViewportClient* InViewportClient, HHitProxy* HitProxy, FInputClick Click);
+	virtual bool IsHitByClick(HHitProxy* HitProx);
+	virtual bool BoxSelect(FBox& InBox, bool InSelect = true);
+	virtual bool FrustumSelect(const FConvexVolume& InFrustum, FEditorViewportClient* InViewportClient, bool InSelect = true) ;
+
+	//Get calculated center of all selections into one position, return false if not selected
+	virtual bool IsAnythingSelected(FVector& OutVectorPosition)const;
+	//Get each selected item's position, return false if not selected.
+	virtual bool IsAnythingSelected(TArray<FVector>& OutVectorPositions) const;
+	virtual bool IsAnythingSelected() const;
+	virtual void SelectNone();
+	virtual bool IsSelected(const FGuid& Key) const;
+	virtual bool IsAlwaysVisible(const FGuid Key) const;
+
+	virtual void AddTrail(const FGuid& Key, TUniquePtr<FTrail>&& TrailPtr);
+	virtual void RemoveTrail(const FGuid& Key);
+
+	virtual bool StartTracking();
+	virtual bool ApplyDelta(const FVector& Pos, const FRotator& Rot, const FVector& WidgetLocation);
+	virtual bool EndTracking();
+
+	virtual void TranslateSelectedKeys(bool bRight);
+	virtual void DeleteSelectedKeys();
+
+	const TRange<double>& GetViewRange() const { return ViewRange; }
+	const TMap<FGuid, TUniquePtr<FTrail>>& GetAllTrails() const { return AllTrails; }
+
+	const TMap<FString, FTimespan>& GetTimingStats() const { return TimingStats; };
+	TMap<FString, FTimespan>& GetTimingStats() { return TimingStats; }
+
+	FTrailVisibilityManager& GetVisibilityManager() { return VisibilityManager; }
+
+protected:
+	void RemoveTrailIfNotAlwaysVisible(const FGuid& Key);
+
+	void OpenContextMenu(const FGuid& TrailGuid);
+protected:
+
+
+	TRange<double> ViewRange;
+	TRange<double> EvalRange;
+	TArray<double> EvalTimesArr;
+	double SecondsPerSegment;
+
+	TRange<double> LastEvalRange;
+	double LastSecondsPerSegment;
+
+	TMap<FGuid, TUniquePtr<FTrail>> AllTrails;
+
+	TMap<FString, FTimespan> TimingStats;
+
+	FTrailVisibilityManager VisibilityManager;
+};
+
+} // namespace MovieScene
+} // namespace UE
diff --git a/Engine/Plugins/Editor/SequencerAnimTools/Source/SequencerAnimTools/Public/TrajectoryCache.h b/Engine/Plugins/Editor/SequencerAnimTools/Source/SequencerAnimTools/Public/TrajectoryCache.h
new file mode 100644
--- /dev/null
+++ b/Engine/Plugins/Editor/SequencerAnimTools/Source/SequencerAnimTools/Public/TrajectoryCache.h
@@ -0,0 +1,138 @@
+// Copyright Epic Games, Inc. All Rights Reserved.
+
+#pragma once
+
+#include "Containers/ArrayView.h"
+#include "Containers/Map.h"
+#include "UObject/WeakObjectPtrTemplates.h"
+#include "Misc/Optional.h"
+#include "Math/Range.h"
+#include "Math/Transform.h"
+
+namespace UE
+{
+namespace SequencerAnimTools
+{
+
+bool RangesContain(const TSet<TRange<double>>& Ranges, const double Time);
+
+
+// Intermediate range reprsentation structure
+class FTrailEvaluateTimes
+{
+public:
+	FTrailEvaluateTimes()
+		: EvalTimes()
+		, Spacing(0.1)
+		, Range(TRange<double>::Empty())
+	{}
+
+	FTrailEvaluateTimes(TArrayView<double> InEvalTimes, double InSpacing) 
+		: EvalTimes(InEvalTimes)
+		, Spacing(InSpacing)
+		, Range(TRange<double>(InEvalTimes[0], InEvalTimes.Last()))
+	{}
+
+	TArrayView<double> EvalTimes;
+	double Spacing;
+	TRange<double> Range;
+};
+
+class FTrajectoryCache
+{
+public:
+	FTrajectoryCache()
+		: Default(FTransform::Identity)
+		, TrackRange()
+		, Spacing()
+	{}
+
+	FTrajectoryCache(const double InSpacing, const TRange<double>& InTrackRange, const FTransform& InDefault)
+		: Default(InDefault)
+		, TrackRange(TRange<double>(FMath::FloorToDouble(InTrackRange.GetLowerBoundValue() / InSpacing)* InSpacing, FMath::FloorToDouble(InTrackRange.GetUpperBoundValue() / InSpacing)* InSpacing)) // snap to spacing
+		, Spacing(InSpacing)
+
+	{}
+
+	virtual ~FTrajectoryCache() {}
+
+	// alternatively, GetEvaluateTimes(const FTrailEvaluateTimes&) and GetInterpEvaluateTimes(const FTrailEvaluateTimes&)
+	virtual const FTransform& Get(const double InTime) const = 0;
+	virtual FTransform GetInterp(const double InTime) const = 0;
+	virtual void GetInterp(const double InTime, FTransform& OutTransform, FTransform& OutParentTransform) const{};
+	virtual const FTransform& GetParent(const double InTime) const { return FTransform::Identity; }
+
+	virtual TArray<double> GetAllTimesInRange(const TRange<double>& InRange) const = 0;
+	
+	virtual void Set(const double InTime, const FTransform& InValue) = 0;
+	virtual void SetTransforms(TArray<FTransform>& InTrajectoryCache, TArray<FTransform>& InParentCache) = 0;
+	// Optionally implemented
+	virtual const FTransform& GetDefault() const { return Default; };
+	virtual void UpdateCacheTimes(const FTrailEvaluateTimes& InEvalTimes);
+	
+	const TRange<double>& GetTrackRange() const { return TrackRange; }
+
+protected:
+	FTransform Default;
+	TRange<double> TrackRange;
+	double Spacing;
+};
+
+class FArrayTrajectoryCache : public FTrajectoryCache
+{
+public:
+
+	// for FRootTrail
+	FArrayTrajectoryCache()
+		: FTrajectoryCache()
+		, TrajectoryCache()
+
+	{}
+
+	FArrayTrajectoryCache(const double InSpacing, const TRange<double>& InTrackRange, FTransform InDefault = FTransform::Identity)
+		: FTrajectoryCache(InSpacing,InTrackRange,InDefault)
+		, TrajectoryCache()
+
+	{
+		TrajectoryCache.SetNumUninitialized(int32(InTrackRange.Size<double>() / InSpacing) + 1);
+	}
+
+	virtual const FTransform& Get(const double InTime) const override 
+	{ 
+		return TrajectoryCache.Num() > 0 ? TrajectoryCache[FMath::Clamp(int32((InTime - TrackRange.GetLowerBoundValue()) / Spacing), 0, TrajectoryCache.Num() - 1)] : Default; 
+	}
+
+	virtual FTransform GetInterp(const double InTime) const override;
+	virtual void GetInterp(const double InTime, FTransform& OutTransform, FTransform& OutParentTransform) const override;
+
+	virtual const FTransform& GetParent(const double InTime) const override
+	{
+		return ParentTrajectoryCache.Num() > 0 ? ParentTrajectoryCache[FMath::Clamp(int32((InTime - TrackRange.GetLowerBoundValue()) / Spacing), 0, ParentTrajectoryCache.Num() - 1)] : FTransform::Identity;
+	}
+
+
+	virtual TArray<double> GetAllTimesInRange(const TRange<double>& InRange) const override;
+
+	virtual void Set(const double InTime, const FTransform& InValue) override 
+	{ 
+		if (TrajectoryCache.Num() > 0)
+		{
+			TrajectoryCache[FMath::Clamp(int32((InTime - TrackRange.GetLowerBoundValue()) / Spacing), 0, TrajectoryCache.Num() - 1)] = InValue;
+		}
+	}
+
+	virtual void SetTransforms(TArray<FTransform>& InTrajectoryCache, TArray<FTransform>& InParentCache)
+	{
+		TrajectoryCache = InTrajectoryCache;
+		ParentTrajectoryCache = InParentCache;
+	}
+
+	
+private:
+	TArray<FTransform> TrajectoryCache;
+	TArray<FTransform> ParentTrajectoryCache;
+
+};
+
+} // namespace MovieScene
+} // namespace UE
diff --git a/Engine/Plugins/Editor/SequencerAnimTools/Source/SequencerAnimTools/Public/TrajectoryDrawInfo.h b/Engine/Plugins/Editor/SequencerAnimTools/Source/SequencerAnimTools/Public/TrajectoryDrawInfo.h
new file mode 100644
--- /dev/null
+++ b/Engine/Plugins/Editor/SequencerAnimTools/Source/SequencerAnimTools/Public/TrajectoryDrawInfo.h
@@ -0,0 +1,77 @@
+// Copyright Epic Games, Inc. All Rights Reserved.
+
+#pragma once
+
+#include "Misc/Guid.h"
+#include "Containers/Map.h"
+#include "Containers/Array.h"
+#include "Math/Color.h"
+#include "Math/Range.h"
+#include "Math/Vector.h"
+#include "Misc/Optional.h"
+
+#include "SceneView.h"
+#include "UnrealClient.h"
+
+namespace UE
+{
+namespace SequencerAnimTools
+{
+
+class FTrail;  
+class FMapTrajectoryCache;
+class FTrajectoryCache;
+
+class FTrailScreenSpaceTransform
+{
+public:
+	FTrailScreenSpaceTransform(const FSceneView* InView, FViewport* InViewport, const float InDPIScale = 1.0f)
+		: View(InView)
+		, HalfScreenSize((InViewport->GetSizeXY().X / InDPIScale) * 0.5, (InViewport->GetSizeXY().Y / InDPIScale) * 0.5)
+	{}
+
+	TOptional<FVector2D> ProjectPoint(const FVector& Point) const;
+
+private:
+	const FSceneView* View;
+	FVector2D HalfScreenSize;
+};
+
+class FTrajectoryDrawInfo
+{
+public:
+
+	FTrajectoryDrawInfo(const FLinearColor& InColor, FTrajectoryCache* InTrajectoryCache)
+		: Color(InColor)
+		, CachedViewRange(TRange<double>::Empty())
+		, TrajectoryCache(InTrajectoryCache)
+	{}
+
+	virtual ~FTrajectoryDrawInfo() {}
+
+	struct FDisplayContext
+	{
+		FGuid YourNode;
+		const FTrailScreenSpaceTransform& ScreenSpaceTransform;
+		double SecondsPerTick;
+		const TRange<double>& TimeRange;
+		double SecondsPerSegment;
+	};
+
+	void GetTrajectoryPointsForDisplay(const FDisplayContext& InDisplayContext, TArray<FVector>& OutPoints, TArray<double>& OutSeconds);
+	void GetTickPointsForDisplay(const FDisplayContext& InDisplayContext, TArray<FVector2D>& OutTicks, TArray<FVector2D>& OutTickTangents);
+
+	// Optionally implemented methods
+	void SetColor(const FLinearColor& InColor) { Color = InColor; }
+	FLinearColor GetColor() const { return Color; }
+
+	virtual const TRange<double>& GetCachedViewRange() const { return CachedViewRange; }
+protected:
+	FLinearColor Color;
+	TRange<double> CachedViewRange;
+	FTrajectoryCache* TrajectoryCache;
+};
+
+
+} // namespace MovieScene
+} // namespace UE
diff --git a/Engine/Plugins/Editor/SequencerAnimTools/Source/SequencerAnimTools/SequencerAnimTools.Build.cs b/Engine/Plugins/Editor/SequencerAnimTools/Source/SequencerAnimTools/SequencerAnimTools.Build.cs
new file mode 100644
--- /dev/null
+++ b/Engine/Plugins/Editor/SequencerAnimTools/Source/SequencerAnimTools/SequencerAnimTools.Build.cs
@@ -0,0 +1,52 @@
+// Copyright Epic Games, Inc. All Rights Reserved.
+
+using UnrealBuildTool;
+
+public class SequencerAnimTools: ModuleRules
+{
+	public SequencerAnimTools (ReadOnlyTargetRules Target) : base(Target)
+	{
+		PCHUsage = ModuleRules.PCHUsageMode.UseExplicitOrSharedPCHs;
+		
+		PublicDependencyModuleNames.AddRange(
+			new string[]
+			{
+				"Core",
+                "Slate",
+                "SlateCore",
+                "EditorStyle"
+			}
+		);
+			
+		
+		PrivateDependencyModuleNames.AddRange(
+			new string[]
+			{
+				"CoreUObject",
+				"Engine",
+				"Slate",
+				"SlateCore",
+				"InputCore",
+				"UnrealEd",
+				"LevelEditor",
+				"AppFramework",
+				
+				"EditorFramework",
+                "EditorInteractiveToolsFramework",
+                "InteractiveToolsFramework",
+                "ViewportInteraction",
+
+                "MovieScene",
+                "MovieSceneTracks",
+                "MovieSceneTools",
+                "Sequencer",
+				"LevelSequence",
+				"LevelSequenceEditor",
+				"SequencerScripting",
+
+				"ControlRig",
+				"ControlRigEditor"
+			}
+        );
+	}
+}
diff --git a/Engine/Source/Editor/MovieSceneTools/Private/EditModes/SkeletalAnimationTrackEditMode.cpp b/Engine/Source/Editor/MovieSceneTools/Private/EditModes/SkeletalAnimationTrackEditMode.cpp
--- a/Engine/Source/Editor/MovieSceneTools/Private/EditModes/SkeletalAnimationTrackEditMode.cpp
+++ b/Engine/Source/Editor/MovieSceneTools/Private/EditModes/SkeletalAnimationTrackEditMode.cpp
@@ -469,20 +469,12 @@ bool FSkeletalAnimationTrackEditMode::StartTracking(FEditorViewportClient* InVie
 
 bool FSkeletalAnimationTrackEditMode::UsesTransformWidget() const
 {
-	if (IsSomethingSelected())
-	{
-		return true;
-	}
-	return FEdMode::UsesTransformWidget();
+	return IsSomethingSelected();
 }
 
 bool FSkeletalAnimationTrackEditMode::UsesTransformWidget(UE::Widget::EWidgetMode CheckMode) const
 {
-	if (IsSomethingSelected())
-	{
-		return true;
-	}
-	return FEdMode::UsesTransformWidget(CheckMode);
+	return IsSomethingSelected();
 }
 
 FVector FSkeletalAnimationTrackEditMode::GetWidgetLocation() const
@@ -517,7 +509,7 @@ bool FSkeletalAnimationTrackEditMode::HandleClick(FEditorViewportClient* InViewp
 			}
 		}
 	}
-	return FEdMode::HandleClick(InViewportClient, HitProxy, Click);
+	return false;
 }
 
 bool FSkeletalAnimationTrackEditMode::IsSomethingSelected() const
diff --git a/Engine/Source/Editor/MovieSceneTools/Private/MovieSceneToolHelpers.cpp b/Engine/Source/Editor/MovieSceneTools/Private/MovieSceneToolHelpers.cpp
--- a/Engine/Source/Editor/MovieSceneTools/Private/MovieSceneToolHelpers.cpp
+++ b/Engine/Source/Editor/MovieSceneTools/Private/MovieSceneToolHelpers.cpp
@@ -940,7 +940,7 @@ void MovieSceneToolHelpers::MovieSceneTranslatorLogOutput(FMovieSceneTranslator
 	}
 }
 
-static FGuid GetHandleToObject(UObject* InObject, UMovieSceneSequence* InSequence, IMovieScenePlayer* Player, FMovieSceneSequenceIDRef TemplateID)
+static FGuid GetHandleToObject(UObject* InObject, UMovieSceneSequence* InSequence, IMovieScenePlayer* Player, FMovieSceneSequenceIDRef TemplateID, bool bCreateIfMissing = true)
 {
 	UMovieScene* MovieScene = InSequence->GetMovieScene();
 
@@ -973,8 +973,11 @@ static FGuid GetHandleToObject(UObject* InObject, UMovieSceneSequence* InSequenc
 		return PropertyOwnerGuid;
 	}
 
-	// Otherwise, create a possessable for this object. Note this will handle creating the parent possessables if this is a component.
-	PropertyOwnerGuid = InSequence->CreatePossessable(InObject);
+	if (bCreateIfMissing)
+	{
+		// Otherwise, create a possessable for this object. Note this will handle creating the parent possessables if this is a component.
+		PropertyOwnerGuid = InSequence->CreatePossessable(InObject);
+	}
 	
 	return PropertyOwnerGuid;
 }
@@ -4044,8 +4047,48 @@ USkeletalMeshComponent* MovieSceneToolHelpers::AcquireSkeletalMeshFromObject(UOb
 	}
 	return nullptr;
 }
+static void GetSequenceSceneComponentWorldTransforms(USceneComponent* SceneComponent, IMovieScenePlayer* Player, UMovieSceneSequence* InSequence, FMovieSceneSequenceIDRef Template, const TArray<FFrameNumber>& Frames, TArray<FTransform>& OutTransforms)
+{
+
+	// Hack: static system interrogator for now to avoid re-allocating UObjects all the time
+	static UE::MovieScene::FSystemInterrogator Interrogator;
+	Interrogator.Reset();
+
+	Interrogator.ImportTransformHierarchy(SceneComponent, Player, Template);
+
+	if (Frames[0] < Frames[Frames.Num() - 1])
+	{
+		for (const FFrameNumber& FrameNumber : Frames)
+		{
+			const FFrameTime FrameTime(FrameNumber);
+			Interrogator.AddInterrogation(FrameTime);
+		}
+		Interrogator.Update();
+
+		Interrogator.QueryWorldSpaceTransforms(SceneComponent, OutTransforms);
+	}
+	else //time is backward we need to do swap things around
+	{
+		for (int32 Index = Frames.Num() - 1; Index >= 0; --Index)
+		{
+			const FFrameTime FrameTime(Frames[Index]);
+			Interrogator.AddInterrogation(FrameTime);
+		}
+		Interrogator.Update();
 
-static void GetSequencerActorWorldTransforms(ISequencer* Sequencer, const FActorForWorldTransforms& ActorSelection, const TArray<FFrameNumber>& Frames, TArray<FTransform>& OutTransforms)
+		TArray<FTransform> WorldTransforms;
+		Interrogator.QueryWorldSpaceTransforms(SceneComponent, WorldTransforms);
+		OutTransforms.SetNum(0);
+		OutTransforms.Reserve(WorldTransforms.Num());
+		for (int32 Index2 = WorldTransforms.Num() - 1; Index2 >= 0; --Index2)
+		{
+			OutTransforms.Add(WorldTransforms[Index2]);
+		}
+	}
+	Interrogator.Reset();
+}
+
+static void GetSequencerActorWorldTransforms(IMovieScenePlayer* Player, UMovieSceneSequence* InSequence, FMovieSceneSequenceIDRef Template, const FActorForWorldTransforms& ActorSelection, const TArray<FFrameNumber>& Frames, TArray<FTransform>& OutTransforms)
 {
 	if (AActor* Actor = ActorSelection.Actor.Get())
 	{
@@ -4060,11 +4103,10 @@ static void GetSequencerActorWorldTransforms(ISequencer* Sequencer, const FActor
 		if (ActorSelection.SocketName != NAME_None && SkelMeshComp)
 		{
 
-			if (UMovieScene* MovieScene = Sequencer->GetFocusedMovieSceneSequence()->GetMovieScene())
+			if (UMovieScene* MovieScene = InSequence->GetMovieScene())
 			{
 				OutTransforms.SetNum(Frames.Num());
 
-				FMovieSceneSequenceIDRef Template = Sequencer->GetFocusedTemplateID();
 				FMovieSceneSequenceTransform RootToLocalTransform;
 
 				FFrameRate TickResolution = MovieScene->GetTickResolution();
@@ -4092,12 +4134,12 @@ static void GetSequencerActorWorldTransforms(ISequencer* Sequencer, const FActor
 						}
 					}
 
-					FMovieSceneContext Context = FMovieSceneContext(FMovieSceneEvaluationRange(GlobalTime, TickResolution), Sequencer->GetPlaybackStatus()).SetHasJumped(true);
+					FMovieSceneContext Context = FMovieSceneContext(FMovieSceneEvaluationRange(GlobalTime, TickResolution), Player->GetPlaybackStatus()).SetHasJumped(true);
 					if (Index == 0) // similar with baking first time in we need to evaluate twice (think due to double buffering that happens with skel mesh components).
 					{
-						Sequencer->GetEvaluationTemplate().Evaluate(Context, *Sequencer);
+						Player->GetEvaluationTemplate().Evaluate(Context, *Player);
 					}
-					Sequencer->GetEvaluationTemplate().Evaluate(Context, *Sequencer);
+					Player->GetEvaluationTemplate().Evaluate(Context, *Player);
 
 					for (IMovieSceneToolsAnimationBakeHelper* BakeHelper : BakeHelpers)
 					{
@@ -4136,49 +4178,26 @@ static void GetSequencerActorWorldTransforms(ISequencer* Sequencer, const FActor
 		}
 		else //no attached skelmesh socket so use Interrogator
 		{
-			// Hack: static system interrogator for now to avoid re-allocating UObjects all the time
-			static UE::MovieScene::FSystemInterrogator Interrogator;
-			Interrogator.Reset();
-
 			USceneComponent* SceneComponent = Actor->GetRootComponent();
-
-			Interrogator.ImportTransformHierarchy(SceneComponent, Sequencer, Sequencer->GetFocusedTemplateID());
-
-			if (Frames[0] < Frames[Frames.Num() - 1])
+			if (SceneComponent)
 			{
-				for (const FFrameNumber& FrameNumber : Frames)
-				{
-					const FFrameTime FrameTime(FrameNumber);
-					Interrogator.AddInterrogation(FrameTime);
-				}
-				Interrogator.Update();
-
-				Interrogator.QueryWorldSpaceTransforms(SceneComponent, OutTransforms);
+				GetSequenceSceneComponentWorldTransforms(SceneComponent, Player, InSequence, Template, Frames, OutTransforms);
 			}
-			else //time is backward we need to do swap things around
-			{
-				for (int32 Index = Frames.Num() - 1; Index >= 0; --Index)
-				{
-					const FFrameTime FrameTime(Frames[Index]);
-					Interrogator.AddInterrogation(FrameTime);
-				}
-				Interrogator.Update();
-
-				TArray<FTransform> WorldTransforms;
-				Interrogator.QueryWorldSpaceTransforms(SceneComponent, WorldTransforms);
-				OutTransforms.SetNum(0);
-				OutTransforms.Reserve(WorldTransforms.Num());
-				for (int32 Index2 = WorldTransforms.Num() - 1; Index2 >= 0; --Index2)
-				{
-					OutTransforms.Add(WorldTransforms[Index2]);
-				}
-			}
-			Interrogator.Reset();
+			
+		}
+	}
+	else //no actor so check to see if there's a scene component
+	{
+		USceneComponent* SceneComponent = ActorSelection.Component.IsValid() ? Cast<USceneComponent>(ActorSelection.Component.Get()) : nullptr;
+		if (SceneComponent)
+		{
+			GetSequenceSceneComponentWorldTransforms(SceneComponent, Player, InSequence, Template, Frames, OutTransforms);
 		}
+
 	}
 }
 
-static void GetNonSequencerActorWorldTransforms(ISequencer* Sequencer, const FActorForWorldTransforms& ActorSelection, const TArray<FFrameNumber>& Frames, TArray<FTransform>& OutTransforms)
+static void GetNonSequencerActorWorldTransforms(IMovieScenePlayer* Player, UMovieSceneSequence* InSequence, FMovieSceneSequenceIDRef Template, const FActorForWorldTransforms& ActorSelection, const TArray<FFrameNumber>& Frames, TArray<FTransform>& OutTransforms)
 {
 	FName SocketName = ActorSelection.SocketName;
 	AActor* Actor = ActorSelection.Actor.Get();
@@ -4188,10 +4207,10 @@ static void GetNonSequencerActorWorldTransforms(ISequencer* Sequencer, const FAc
 	{
 		do
 		{
-			FGuid ActorHandle = Sequencer->GetHandleToObject(Actor, false);
+			FGuid ActorHandle = GetHandleToObject(Actor, InSequence, Player, Template,false);
 			if (ActorHandle.IsValid())
 			{
-				GetSequencerActorWorldTransforms(Sequencer, NewActorSelection, Frames, OutTransforms);
+				GetSequencerActorWorldTransforms(Player,InSequence,Template, NewActorSelection, Frames, OutTransforms);
 				for (FTransform& OutTransform : OutTransforms)
 				{
 					OutTransform = WorldTransform * OutTransform;
@@ -4225,14 +4244,62 @@ static void GetNonSequencerActorWorldTransforms(ISequencer* Sequencer, const FAc
 
 void MovieSceneToolHelpers::GetActorWorldTransforms(ISequencer* Sequencer, const FActorForWorldTransforms& ActorSelection, const TArray<FFrameNumber>& Frames, TArray<FTransform>& OutWorldTransforms)
 {
+	FMovieSceneSequenceIDRef Template = Sequencer->GetFocusedTemplateID();
+
 	FGuid ObjectHandle = Sequencer->GetHandleToObject(ActorSelection.Actor.Get(), false);
 	if (ObjectHandle.IsValid())
 	{
-		GetSequencerActorWorldTransforms(Sequencer, ActorSelection, Frames, OutWorldTransforms);
+		GetSequencerActorWorldTransforms(Sequencer, Sequencer->GetFocusedMovieSceneSequence(),Template, ActorSelection, Frames, OutWorldTransforms);
+	}
+	else
+	{
+		ObjectHandle = Sequencer->GetHandleToObject(ActorSelection.Component.Get(), false);
+		if (ObjectHandle.IsValid())
+		{
+			GetSequencerActorWorldTransforms(Sequencer, Sequencer->GetFocusedMovieSceneSequence(), Template, ActorSelection, Frames, OutWorldTransforms);
+		}
+		else
+		{
+			GetNonSequencerActorWorldTransforms(Sequencer, Sequencer->GetFocusedMovieSceneSequence(), Template, ActorSelection, Frames, OutWorldTransforms);
+		}
+	}
+}
+
+void MovieSceneToolHelpers::SetOrAddKey(TMovieSceneChannelData<FMovieSceneFloatValue>& ChannelData, FFrameNumber Time, float Value)
+{
+	int32 ExistingIndex = ChannelData.FindKey(Time);
+	if (ExistingIndex != INDEX_NONE)
+	{
+		FMovieSceneFloatValue& FloatValue = ChannelData.GetValues()[ExistingIndex]; //-V758
+		FloatValue.Value = Value;
 	}
 	else
 	{
-		GetNonSequencerActorWorldTransforms(Sequencer, ActorSelection, Frames, OutWorldTransforms);
+		FMovieSceneFloatValue NewKey(Value);
+		ERichCurveTangentWeightMode WeightedMode = RCTWM_WeightedNone;
+		NewKey.InterpMode = ERichCurveInterpMode::RCIM_Cubic;
+		NewKey.TangentMode = ERichCurveTangentMode::RCTM_Auto;
+		NewKey.Tangent.ArriveTangent = 0.0f;
+		NewKey.Tangent.LeaveTangent = 0.0f;
+		NewKey.Tangent.TangentWeightMode = WeightedMode;
+		NewKey.Tangent.ArriveTangentWeight = 0.0f;
+		NewKey.Tangent.LeaveTangentWeight = 0.0f;
+		ChannelData.AddKey(Time, NewKey);
 	}
 }
 
+void MovieSceneToolHelpers::GetActorWorldTransforms(IMovieScenePlayer* Player, UMovieSceneSequence* InSequence, FMovieSceneSequenceIDRef Template, const FActorForWorldTransforms& ActorSelection, const TArray<FFrameNumber>& Frames, TArray<FTransform>& OutWorldTransforms)
+{
+	FGuid ActorHandle = GetHandleToObject(ActorSelection.Actor.Get(), InSequence, Player, Template,false);
+	FGuid ComponentHandle = GetHandleToObject(ActorSelection.Component.Get(), InSequence, Player, Template,false);
+	if (ActorHandle.IsValid() || ComponentHandle.IsValid())
+	{
+		GetSequencerActorWorldTransforms(Player, InSequence, Template, ActorSelection, Frames, OutWorldTransforms);
+	}
+	else
+	{
+		GetNonSequencerActorWorldTransforms(Player, InSequence, Template, ActorSelection, Frames, OutWorldTransforms);
+	}
+}
+
+
diff --git a/Engine/Source/Editor/MovieSceneTools/Private/Sections/TransformPropertySection.cpp b/Engine/Source/Editor/MovieSceneTools/Private/Sections/TransformPropertySection.cpp
--- a/Engine/Source/Editor/MovieSceneTools/Private/Sections/TransformPropertySection.cpp
+++ b/Engine/Source/Editor/MovieSceneTools/Private/Sections/TransformPropertySection.cpp
@@ -61,6 +61,26 @@ void FTransformSection::BuildSectionContextMenu(FMenuBuilder& MenuBuilder, const
 			})
 		);
 	};
+	MenuBuilder.BeginSection(NAME_None, LOCTEXT("MotionTrailOptions", "Motion Trail Options"));
+	{
+
+		FUIAction Action_ConfigureTexturesForSprites(
+			FExecuteAction::CreateLambda([]
+				{
+					FName ControlRigMotionTrailTab("ControlRigMotionTrailTab");
+					FGlobalTabmanager::Get()->TryInvokeTab(ControlRigMotionTrailTab);
+				}
+		));
+
+		MenuBuilder.AddMenuEntry(
+			LOCTEXT("ShowMotionTrailOptions", "Show Motion Trail Options"),
+			LOCTEXT("ShowMotionTrailOptionsTooltip", "Show various motion trail options"),
+			FSlateIcon(),
+			Action_ConfigureTexturesForSprites,
+			NAME_None,
+			EUserInterfaceActionType::Button);
+	}
+	MenuBuilder.EndSection();
 
 	MenuBuilder.BeginSection(NAME_None, LOCTEXT("TransformChannelsText", "Active Channels"));
 	{
diff --git a/Engine/Source/Editor/MovieSceneTools/Public/MovieSceneToolHelpers.h b/Engine/Source/Editor/MovieSceneTools/Public/MovieSceneToolHelpers.h
--- a/Engine/Source/Editor/MovieSceneTools/Public/MovieSceneToolHelpers.h
+++ b/Engine/Source/Editor/MovieSceneTools/Public/MovieSceneToolHelpers.h
@@ -539,6 +539,25 @@ public:
 	*/
 	static void GetActorWorldTransforms(ISequencer* Sequencer, const FActorForWorldTransforms& Actors, const TArray<FFrameNumber>& Frames, TArray<FTransform>& OutWorldTransforms);
 
+	/* Set or add a key onto a float channel.
+	* @param ChannelData Channel to set or add
+	* @param Time Frame to add or set the value
+	* @param Value  Value to Set
+	*/
+	static void SetOrAddKey(TMovieSceneChannelData<FMovieSceneFloatValue>& ChannelData, FFrameNumber Time, float Value);
+
+	
+	/*
+	*  Get an actors world transforms at the specified times using a player
+	* @param Player Player to evaluate
+	* @param InSequence  Sequence to evaluate
+	* @param Template  Sequence ID of the template to play
+    * @param ActorForWorldTransforms The actor and possible component and socket that we want to get the world transforms for.
+	* @param Frames The times we want to get the world transforms
+	* @param OutWorldTransforms The calculated world transforms, one for each specified frame.
+	*/
+	static void GetActorWorldTransforms(IMovieScenePlayer* Player, UMovieSceneSequence* InSequence, FMovieSceneSequenceIDRef Template,const FActorForWorldTransforms& Actors, const TArray<FFrameNumber>& Frames, TArray<FTransform>& OutWorldTransforms);
+
 };
 
 // Helper to make spawnables persist throughout the export process and then restore properly afterwards
diff --git a/Engine/Source/Editor/Sequencer/Private/LevelEditorSequencerIntegration.cpp b/Engine/Source/Editor/Sequencer/Private/LevelEditorSequencerIntegration.cpp
--- a/Engine/Source/Editor/Sequencer/Private/LevelEditorSequencerIntegration.cpp
+++ b/Engine/Source/Editor/Sequencer/Private/LevelEditorSequencerIntegration.cpp
@@ -558,10 +558,14 @@ void FLevelEditorSequencerIntegration::UpdateDetails(bool bForceRefresh)
 
 void FLevelEditorSequencerIntegration::ActivateSequencerEditorMode()
 {
+
 	// Release the sequencer mode if we already enabled it
 	FName ResourceName("SequencerMode");
 	AcquiredResources.Release(ResourceName);
 
+	FEditorModeID ModeID = TEXT("SequencerToolsEditMode");
+	GLevelEditorModeTools().ActivateMode(ModeID);
+
 	GLevelEditorModeTools().ActivateMode( FSequencerEdMode::EM_SequencerMode );
 	FSequencerEdMode* SequencerEdMode = (FSequencerEdMode*)GLevelEditorModeTools().GetActiveMode(FSequencerEdMode::EM_SequencerMode);
 
@@ -578,6 +582,13 @@ void FLevelEditorSequencerIntegration::ActivateSequencerEditorMode()
 	AcquiredResources.Add(
 		ResourceName,
 		[] {
+
+			FEditorModeID ModeID = TEXT("SequencerToolsEditMode");
+
+			if (GLevelEditorModeTools().IsModeActive(ModeID))
+			{
+				GLevelEditorModeTools().DeactivateMode(ModeID);
+			}
 			if (GLevelEditorModeTools().IsModeActive(FSequencerEdMode::EM_SequencerMode))
 			{
 				GLevelEditorModeTools().DeactivateMode(FSequencerEdMode::EM_SequencerMode);
diff --git a/Engine/Source/Editor/Sequencer/Private/SSequencerSection.cpp b/Engine/Source/Editor/Sequencer/Private/SSequencerSection.cpp
--- a/Engine/Source/Editor/Sequencer/Private/SSequencerSection.cpp
+++ b/Engine/Source/Editor/Sequencer/Private/SSequencerSection.cpp
@@ -1028,8 +1028,8 @@ int32 SSequencerSection::OnPaint( const FPaintArgs& Args, const FGeometry& Allot
 	const bool bLocked = SectionObject->IsLocked() || SectionObject->IsReadOnly();
 
 	bool bSetSectionToKey = false;
-
-	if (Track && Track->GetSectionToKey() == SectionObject)
+    // Only show section to key border if we have more than one section
+	if (Track && Track->GetAllSections().Num() > 1 && Track->GetSectionToKey() == SectionObject)
 	{
 		bSetSectionToKey = true;
 	}
diff --git a/Engine/Source/Editor/Sequencer/Private/SequencerEdMode.cpp b/Engine/Source/Editor/Sequencer/Private/SequencerEdMode.cpp
--- a/Engine/Source/Editor/Sequencer/Private/SequencerEdMode.cpp
+++ b/Engine/Source/Editor/Sequencer/Private/SequencerEdMode.cpp
@@ -29,6 +29,7 @@
 #include "EntitySystem/Interrogation/MovieSceneInterrogationLinker.h"
 #include "SequencerSectionPainter.h"
 #include "MovieSceneToolHelpers.h"
+#include "Tools/MotionTrailOptions.h"
 
 const FEditorModeID FSequencerEdMode::EM_SequencerMode(TEXT("EM_SequencerMode"));
 
@@ -37,6 +38,13 @@ static TAutoConsoleVariable<bool> CVarDrawMeshTrails(
 	true,
 	TEXT("Toggle to show or hide Level Sequencer VR Editor trails"));
 
+//if true still use old motion trails for sequencer objects.
+TAutoConsoleVariable<bool> CVarUseOldSequencerMotionTrails(
+	TEXT("Sequencer.UseOldSequencerTrails"),
+	false,
+	TEXT("If true show old motion trails, if false use new editable motion trails."));
+
+
 namespace UE
 {
 namespace SequencerEdMode
@@ -532,6 +540,11 @@ void FSequencerEdMode::DrawTransformTrack(const TSharedPtr<ISequencer>& Sequence
 
 void FSequencerEdMode::DrawTracks3D(FPrimitiveDrawInterface* PDI)
 {
+
+	if (CVarUseOldSequencerMotionTrails->GetBool() == false)
+	{
+		return;
+	}
 	for (TWeakPtr<FSequencer> WeakSequencer : Sequencers)
 	{
 		TSharedPtr<FSequencer> Sequencer = WeakSequencer.Pin();
diff --git a/Engine/Source/Editor/Sequencer/Private/Tools/SMotionTrailOptions.cpp b/Engine/Source/Editor/Sequencer/Private/Tools/SMotionTrailOptions.cpp
new file mode 100644
--- /dev/null
+++ b/Engine/Source/Editor/Sequencer/Private/Tools/SMotionTrailOptions.cpp
@@ -0,0 +1,47 @@
+// Copyright Epic Games, Inc. All Rights Reserved.
+
+#include "Tools/SMotionTrailOptions.h"
+#include "Tools/MotionTrailOptions.h"
+
+#define LOCTEXT_NAMESPACE "MotionTrail"
+
+void SMotionTrailOptions::Construct(const FArguments& InArgs)
+{
+	UMotionTrailToolOptions* Settings = GetMutableDefault<UMotionTrailToolOptions>();
+	FPropertyEditorModule& PropertyEditor = FModuleManager::LoadModuleChecked<FPropertyEditorModule>("PropertyEditor");
+
+	FDetailsViewArgs DetailsViewArgs;
+	DetailsViewArgs.bShowOptions = false;
+	DetailsViewArgs.bAllowSearch = false;
+	DetailsViewArgs.bShowPropertyMatrixButton = false;
+	DetailsViewArgs.bUpdatesFromSelection = false;
+	DetailsViewArgs.bLockable = false;
+	DetailsViewArgs.bAllowFavoriteSystem = false;
+	DetailsViewArgs.NameAreaSettings = FDetailsViewArgs::HideNameArea;
+	DetailsViewArgs.ViewIdentifier = "MotionTrailOptions";
+
+	DetailsView = PropertyEditor.CreateDetailView(DetailsViewArgs);
+	DetailsView->SetObject(Settings);
+
+	ChildSlot
+		[
+			SNew(SBorder)
+			.HAlign(HAlign_Fill)
+			.VAlign(VAlign_Fill)
+			.Padding(FMargin(10.0,5.0,10.0,5.0))
+			[
+			SNew(SVerticalBox)
+			
+				+ SVerticalBox::Slot()
+				.AutoHeight()
+				.HAlign(HAlign_Fill)
+				[
+					DetailsView.ToSharedRef()
+				]
+			]
+			
+		];
+}
+
+
+#undef LOCTEXT_NAMESPACE
diff --git a/Engine/Source/Editor/Sequencer/Public/Tools/MotionTrailOptions.h b/Engine/Source/Editor/Sequencer/Public/Tools/MotionTrailOptions.h
new file mode 100644
--- /dev/null
+++ b/Engine/Source/Editor/Sequencer/Public/Tools/MotionTrailOptions.h
@@ -0,0 +1,98 @@
+// Copyright Epic Games, Inc. All Rights Reserved.
+
+#pragma once
+#include "CoreMinimal.h"
+#include "MotionTrailOptions.generated.h"
+
+//if true still use old motion trails for sequencer objects.
+extern SEQUENCER_API TAutoConsoleVariable<bool> CVarUseOldSequencerMotionTrails;
+
+
+// TODO: option to make tick size proportional to distance from camera to get a sense of perspective and scale
+UCLASS(config = EditorPerProjectUserSettings)
+class SEQUENCER_API UMotionTrailToolOptions : public UObject
+{
+	GENERATED_BODY()
+
+public:
+	UMotionTrailToolOptions()
+		: bShowTrails(false)
+		, TrailColor(.22,0.15,1.0)
+		, bShowFullTrail(true)
+		, TrailThickness(0.0f)
+		, FramesBefore(10)
+		, FramesAfter(10)
+		, EvalsPerFrame(1)
+		, bShowKeys(true)
+		, bShowFrameNumber(true)
+		, KeyColor(1.0, 1.0, 1.0)
+		, KeySize(10.0f)
+		, bShowTicks(false)
+		, TickColor(0.25,1.0,0.15)
+		, TickSize(5.0)
+		, bLockTicksToFrames(true)
+		, SecondsPerTick(0.1)
+	{}
+
+	UPROPERTY(EditAnywhere, Category = Trail);
+	bool bShowTrails;
+
+	UPROPERTY(EditAnywhere,  Category = Trail)
+	FLinearColor TrailColor;
+
+	UPROPERTY(EditAnywhere, Category = Trail)
+	bool bShowFullTrail;
+
+	UPROPERTY(EditAnywhere, Category = Trail, Meta = (ClampMin = "0.0"))
+	float TrailThickness;
+
+	UPROPERTY(EditAnywhere, Category = Trail, Meta = (EditCondition = "!bShowFullTrail", ClampMin = "0"))
+	int32 FramesBefore;
+
+	UPROPERTY(EditAnywhere, Category = Trail, Meta = (EditCondition = "!bShowFullTrail", ClampMin = "0"))
+	int32 FramesAfter;
+
+	UPROPERTY(EditAnywhere, Category = Trail, Meta = (ClampMin = "1.0"))
+	int32 EvalsPerFrame;
+	
+	UPROPERTY(EditAnywhere, Category = Keys)
+	bool bShowKeys;
+
+	UPROPERTY(EditAnywhere, Category = Keys)
+	bool bShowFrameNumber;
+
+	UPROPERTY(EditAnywhere, Category = Keys)
+	FLinearColor KeyColor;
+
+	UPROPERTY(EditAnywhere, Category = Keys, Meta = (ClampMin = "0.0"))
+	double KeySize;
+
+	UPROPERTY(EditAnywhere, Category = Ticks)
+	bool bShowTicks;
+
+	UPROPERTY(EditAnywhere, Category = Ticks)
+	FLinearColor TickColor;
+
+	UPROPERTY(EditAnywhere, Category = Ticks, Meta = (ClampMin = "0.0"))
+	double TickSize;
+
+	UPROPERTY(EditAnywhere, Category = Ticks)
+	bool bLockTicksToFrames;
+
+	UPROPERTY(EditAnywhere, Category = Ticks, Meta = (EditCondition = "!bLockTicksToFrames", ClampMin = "0.01"))
+	double SecondsPerTick;
+
+	virtual void PostEditChangeProperty(struct FPropertyChangedEvent& PropertyChangedEvent) override
+	{
+		FName PropertyName = (PropertyChangedEvent.Property != NULL) ? PropertyChangedEvent.Property->GetFName() : NAME_None;
+		OnDisplayPropertyChanged.Broadcast(PropertyName);
+	}
+
+	DECLARE_MULTICAST_DELEGATE_OneParam(FOnDisplayPropertyChanged, FName);
+	FOnDisplayPropertyChanged OnDisplayPropertyChanged;
+
+
+	static UMotionTrailToolOptions* GetTrailOptions()  { return GetMutableDefault<UMotionTrailToolOptions>(); }
+
+
+};
\ No newline at end of file
diff --git a/Engine/Source/Editor/Sequencer/Public/Tools/SMotionTrailOptions.h b/Engine/Source/Editor/Sequencer/Public/Tools/SMotionTrailOptions.h
new file mode 100644
--- /dev/null
+++ b/Engine/Source/Editor/Sequencer/Public/Tools/SMotionTrailOptions.h
@@ -0,0 +1,26 @@
+// Copyright Epic Games, Inc. All Rights Reserved.
+/**
+* Hold the View for the Motion Trail Options Widget
+*/
+#pragma once
+
+#include "CoreMinimal.h"
+#include "Widgets/SWidget.h"
+#include "Widgets/DeclarativeSyntaxSupport.h"
+#include "Widgets/SCompoundWidget.h"
+#include "IDetailsView.h"
+
+class SEQUENCER_API SMotionTrailOptions : public SCompoundWidget
+{
+	SLATE_BEGIN_ARGS(SMotionTrailOptions) {}
+
+	SLATE_END_ARGS()
+	~SMotionTrailOptions()
+	{
+	}
+
+	void Construct(const FArguments& InArgs);
+
+	TSharedPtr<IDetailsView> DetailsView;
+};
+
