Static Analysis Fix for \ vs / in #include

diff --git a/Engine/Plugins/Editor/SequencerAnimTools/Source/SequencerAnimTools/Private/SequencerAnimToolsModule.cpp b/Engine/Plugins/Editor/SequencerAnimTools/Source/SequencerAnimTools/Private/SequencerAnimToolsModule.cpp
--- a/Engine/Plugins/Editor/SequencerAnimTools/Source/SequencerAnimTools/Private/SequencerAnimToolsModule.cpp
+++ b/Engine/Plugins/Editor/SequencerAnimTools/Source/SequencerAnimTools/Private/SequencerAnimToolsModule.cpp
@@ -8,10 +8,10 @@
 #include "EdModeInteractiveToolsContext.h"
 #include "UnrealEdGlobals.h"
 
-#include "BaseGizmos\TransformGizmoUtil.h"
+#include "BaseGizmos/TransformGizmoUtil.h"
 #include "MotionTrailTool.h"
 #include "EditPivotTool.h"
-#include "Tools\MotionTrailOptions.h"
+#include "Tools/MotionTrailOptions.h"
 
 #define LOCTEXT_NAMESPACE "FSequencerAnimToolsModule"
 
