Incremental Build Error, missing  #include "LevelEditorSequencerIntegration.h"

diff --git a/Engine/Plugins/Editor/SequencerAnimTools/Source/SequencerAnimTools/Private/SequencerAnimToolsModule.cpp b/Engine/Plugins/Editor/SequencerAnimTools/Source/SequencerAnimTools/Private/SequencerAnimToolsModule.cpp
--- a/Engine/Plugins/Editor/SequencerAnimTools/Source/SequencerAnimTools/Private/SequencerAnimToolsModule.cpp
+++ b/Engine/Plugins/Editor/SequencerAnimTools/Source/SequencerAnimTools/Private/SequencerAnimToolsModule.cpp
@@ -7,6 +7,7 @@
 #include "LevelEditor.h"
 #include "EdModeInteractiveToolsContext.h"
 #include "UnrealEdGlobals.h"
+#include "LevelEditorSequencerIntegration.h"
 
 #include "BaseGizmos/TransformGizmoUtil.h"
 #include "MotionTrailTool.h"
