CIS Fix for SequencerTrailHierrarchy.cpp

diff --git a/Engine/Plugins/Editor/SequencerAnimTools/Source/SequencerAnimTools/Private/Sequencer/SequencerTrailHierarchy.cpp b/Engine/Plugins/Editor/SequencerAnimTools/Source/SequencerAnimTools/Private/Sequencer/SequencerTrailHierarchy.cpp
--- a/Engine/Plugins/Editor/SequencerAnimTools/Source/SequencerAnimTools/Private/Sequencer/SequencerTrailHierarchy.cpp
+++ b/Engine/Plugins/Editor/SequencerAnimTools/Source/SequencerAnimTools/Private/Sequencer/SequencerTrailHierarchy.cpp
@@ -284,7 +284,7 @@ void FSequencerTrailHierarchy::UpdateControlRig(const FTrailEvaluateTimes& Evalu
 		ULevelSequence* LevelSequence = Cast<ULevelSequence>(Sequencer->GetFocusedMovieSceneSequence());
 
 		USceneComponent* Component = Cast<USceneComponent>(ObjectBinding->GetBoundObject());
-		if (Component)
+		if (Component && Player)
 		{
 			AActor* Actor = Component->GetTypedOuter< AActor >();
 			if (Actor)
