Sequencer:  Anim Tools:
Minor fixes found while prepping demo

diff --git a/Engine/Plugins/Editor/SequencerAnimTools/Source/SequencerAnimTools/Private/EditPivotTool.cpp b/Engine/Plugins/Editor/SequencerAnimTools/Source/SequencerAnimTools/Private/EditPivotTool.cpp
--- a/Engine/Plugins/Editor/SequencerAnimTools/Source/SequencerAnimTools/Private/EditPivotTool.cpp
+++ b/Engine/Plugins/Editor/SequencerAnimTools/Source/SequencerAnimTools/Private/EditPivotTool.cpp
@@ -44,10 +44,10 @@ void FEditPivotCommands::RegisterCommands()
 }
 
 
-TMap<TWeakObjectPtr<UControlRig>, FControlRigMappings> USequencerPivotTool::SavedPivotLocations;
-TArray<FControlRigMappings> USequencerPivotTool::LastSelectedObjects;
+FSavedMappings USequencerPivotTool::SavedPivotLocations;
+FLastSelectedObjects USequencerPivotTool::LastSelectedObjects;
 
-static void GetControlRigsAndSequencer(TArray<UControlRig*>& ControlRigs, TWeakPtr<ISequencer>& SequencerPtr, ULevelSequence** LevelSequence)
+static void GetControlRigsAndSequencer(TArray<TWeakObjectPtr<UControlRig>>& ControlRigs, TWeakPtr<ISequencer>& SequencerPtr, ULevelSequence** LevelSequence)
 {
 	*LevelSequence = ULevelSequenceEditorBlueprintLibrary::GetCurrentLevelSequence();
 	if (*LevelSequence)
@@ -57,7 +57,12 @@ static void GetControlRigsAndSequencer(TArray<UControlRig*>& ControlRigs, TWeakP
 		SequencerPtr = LevelSequenceEditor ? LevelSequenceEditor->GetSequencer() : nullptr;
 		if (SequencerPtr.IsValid())
 		{
-			ControlRigs = UControlRigSequencerEditorLibrary::GetVisibleControlRigs();
+			TArray<UControlRig*> TempControlRigs;
+			TempControlRigs = UControlRigSequencerEditorLibrary::GetVisibleControlRigs();
+			for (UControlRig* ControlRig : TempControlRigs)
+			{
+				ControlRigs.Add(ControlRig);
+			}
 		}
 	}
 }
@@ -65,23 +70,29 @@ bool USequencerPivotToolBuilder::CanBuildTool(const FToolBuilderState& SceneStat
 {
 	//only start if we have anything selected  or we selected somethign last time
 	//which will then get selected later once the tool starts.
-	if (USequencerPivotTool::LastSelectedObjects.Num() > 0)
+	if (USequencerPivotTool::LastSelectedObjects.LastSelectedControlRigs.Num() > 0 || USequencerPivotTool::LastSelectedObjects.LastSelectedActors.Num() > 0)
 	{
 		return true;
 	}
 	ULevelSequence* LevelSequence;
-	TArray<UControlRig*> ControlRigs;
+	TArray<TWeakObjectPtr<UControlRig>> ControlRigs;
 	TWeakPtr<ISequencer> SequencerPtr;
 	GetControlRigsAndSequencer(ControlRigs, SequencerPtr, &LevelSequence);
 
-	for (UControlRig* ControlRig : ControlRigs)
+	for (TWeakObjectPtr<UControlRig> ControlRig : ControlRigs)
 	{
-		if (ControlRig && ControlRig->CurrentControlSelection().Num() > 0)
+		if (ControlRig.IsValid() && ControlRig->CurrentControlSelection().Num() > 0)
 		{
 			return true;
 		}
 	}
-	
+
+	TArray<AActor*> SelectedActors;
+	GEditor->GetSelectedActors()->GetSelectedObjects(SelectedActors);
+	if (SelectedActors.Num() > 0)
+	{
+		return true;
+	}
 	
 	return false;
 }
@@ -122,6 +133,12 @@ void USequencerPivotTool::Setup()
 	//when entered we check to see if shift is pressed this can change where we set the pivot on start or reset
 	FModifierKeysState KeyState = FSlateApplication::Get().GetModifierKeys();
 	bShiftPressedWhenStarted = KeyState.IsShiftDown();
+	bCtrlPressedWhenStarted = KeyState.IsControlDown();
+	if (bCtrlPressedWhenStarted)
+	{
+		bShiftPressedWhenStarted = false; //for now turn off shift which means we pivot from the last if control is also presseed
+	}
+
 
 	UInteractiveTool::Setup();
 
@@ -140,14 +157,29 @@ void USequencerPivotTool::Setup()
 	TransformGizmo->SetActiveTarget(TransformProxy, GetToolManager());
 
 	GetControlRigsAndSequencer(ControlRigs, SequencerPtr, &LevelSequence);
+
+	TArray<AActor*> SelectedActors;
+	GEditor->GetSelectedActors()->GetSelectedObjects(SelectedActors);
+	Actors.SetNum(0);
+	for (AActor* Actor : SelectedActors)
+	{
+		Actors.Add(Actor);
+	}
+
 	UpdateTransformAndSelectionOnEntering();
 	UpdateGizmoTransform();
 	UpdateGizmoVisibility();
 	//we get delegates last since we may select something above
-	for (UControlRig* ControlRig : ControlRigs)
+	for (TWeakObjectPtr<UControlRig>& ControlRig : ControlRigs)
 	{
-		ControlRig->ControlSelected().AddUObject(this,&USequencerPivotTool::HandleControlSelected);
+		if (ControlRig.IsValid())
+		{
+			ControlRig->ControlSelected().AddUObject(this, &USequencerPivotTool::HandleControlSelected);
+		}
 	}
+
+	OnEditorSelectionChangedHandle = USelection::SelectionChangedEvent.AddUObject(this, &USequencerPivotTool::OnEditorSelectionChanged);
+
 	SaveLastSelected();
 
 	CommandBindings = MakeShareable(new FUICommandList);
@@ -163,26 +195,44 @@ void USequencerPivotTool::Setup()
 
 }
 
+
 void USequencerPivotTool::ResetPivot()
 {
 	SetGizmoBasedOnSelection(false);
 	UpdateGizmoTransform();
+	SavePivotTransforms();
 }
 
 //Last selected is really the ones that were selected when you entered the tool
 //we use this in case nothing was selected when the tool is active so that we instead select it.
 void USequencerPivotTool::SaveLastSelected()
 {
-	LastSelectedObjects.SetNum(0);
-	for (UControlRig* ControlRig : ControlRigs)
+
+	LastSelectedObjects.LastSelectedControlRigs.SetNum(0);
+	LastSelectedObjects.LastSelectedActors.SetNum(0);
+
+	for (TWeakObjectPtr<UControlRig>& ControlRig : ControlRigs)
 	{
-		TArray<FName> SelectedControls = ControlRig->CurrentControlSelection();
-		for (const FName& Name : SelectedControls)
+		if (ControlRig.IsValid())
+		{
+			TArray<FName> SelectedControls = ControlRig->CurrentControlSelection();
+			for (const FName& Name : SelectedControls)
+			{
+				FControlRigMappings Mapping;
+				Mapping.ControlRig = ControlRig;
+				Mapping.PivotTransforms.Add(Name, GizmoTransform);
+				LastSelectedObjects.LastSelectedControlRigs.Add(Mapping);
+			}
+		}
+	}
+	for (TWeakObjectPtr<AActor>& Actor : Actors)
+	{
+		if (Actor.IsValid())
 		{
-			FControlRigMappings Mapping;
-			Mapping.ControlRig = ControlRig;
-			Mapping.PivotTransforms.Add(Name, GizmoTransform);
-			LastSelectedObjects.Add(Mapping);
+			FActorMappings Mapping;
+			Mapping.Actor = Actor;
+			Mapping.PivotTransform = GizmoTransform;
+			LastSelectedObjects.LastSelectedActors.Add(Mapping);
 		}
 	}
 }
@@ -192,11 +242,11 @@ void USequencerPivotTool::SaveLastSelected()
 void USequencerPivotTool::UpdateTransformAndSelectionOnEntering()
 {
 	//if shift was pressed when we started we don't use saved, this will move pivot to last object
-	bool bhaveSomethingSelected = SetGizmoBasedOnSelection(!bShiftPressedWhenStarted);
-	//okay nothing selected we select the last thing selec
+	bool bhaveSomethingSelected = SetGizmoBasedOnSelection(!bShiftPressedWhenStarted && !bCtrlPressedWhenStarted);
+	//okay nothing selected we select the last thing selected
 	if (bhaveSomethingSelected == false)
 	{
-		for (FControlRigMappings& Mappings : LastSelectedObjects)
+		for (FControlRigMappings& Mappings : LastSelectedObjects.LastSelectedControlRigs)
 		{
 			if (Mappings.ControlRig.IsValid())
 			{
@@ -206,7 +256,14 @@ void USequencerPivotTool::UpdateTransformAndSelectionOnEntering()
 				}
 			}
 		}
-		SetGizmoBasedOnSelection(!bShiftPressedWhenStarted);
+		for (FActorMappings& ActorMappings : LastSelectedObjects.LastSelectedActors)
+		{
+			if (ActorMappings.Actor.IsValid())
+			{
+				GEditor->SelectActor(ActorMappings.Actor.Get(), true, true);
+			}
+		}
+		SetGizmoBasedOnSelection(!bShiftPressedWhenStarted && !bCtrlPressedWhenStarted);
 	}
 }
 
@@ -234,9 +291,15 @@ bool USequencerPivotTool::SetGizmoBasedOnSelection(bool bUseSaved)
 	const FFrameTime FrameTime = Sequencer->GetLocalTime().ConvertTo(TickResolution);
 
 	bool bHaveSomethingSelected = false;
-	for (UControlRig* ControlRig : ControlRigs)
+	FVector AverageLocation(0.0f);
+	int NumLocations = 0;
+	for (TWeakObjectPtr<UControlRig>& ControlRig : ControlRigs)
 	{
-		FControlRigMappings* Mappings = SavedPivotLocations.Find(ControlRig);
+		if (ControlRig.IsValid() == false)
+		{
+			continue;
+		}
+		FControlRigMappings* Mappings = SavedPivotLocations.ControlRigMappings.Find(ControlRig);
 		TArray<FName> SelectedControls = ControlRig->CurrentControlSelection();
 
 		for (const FName& Name : SelectedControls)
@@ -253,11 +316,37 @@ bool USequencerPivotTool::SetGizmoBasedOnSelection(bool bUseSaved)
 			}
 			if (bHasSavedPivot == false)
 			{
-				GizmoTransform = UControlRigSequencerEditorLibrary::GetControlRigWorldTransform(LevelSequence, ControlRig, Name, FrameTime.RoundToFrame(),
+				GizmoTransform = UControlRigSequencerEditorLibrary::GetControlRigWorldTransform(LevelSequence, ControlRig.Get(), Name, FrameTime.RoundToFrame(),
 					ESequenceTimeUnit::TickResolution);
+				AverageLocation += GizmoTransform.GetLocation();
+				NumLocations++;
 			}
 		}
 	}
+	//now do actors, to do need to do them in order somehow if using bUseSaved.... (shift is pressed) and we have multiple
+
+	TArray<AActor*> SelectedActors;
+	GEditor->GetSelectedActors()->GetSelectedObjects(SelectedActors);
+	for (AActor* SelectedActor : SelectedActors)
+	{
+		FActorMappings* Mappings = SavedPivotLocations.ActorMappings.Find(SelectedActor);
+		if (bUseSaved && Mappings)
+		{
+			GizmoTransform = Mappings->PivotTransform;
+		}
+		else
+		{
+			GizmoTransform = UControlRigSequencerEditorLibrary::GetActorWorldTransform(LevelSequence,SelectedActor, FrameTime.RoundToFrame(),
+				ESequenceTimeUnit::TickResolution);
+			AverageLocation += GizmoTransform.GetLocation();
+			NumLocations++;
+		}
+	}
+	if (bCtrlPressedWhenStarted && NumLocations > 1)
+	{
+		AverageLocation /= (float)(NumLocations);
+		GizmoTransform.SetLocation(AverageLocation);
+	}
 	GizmoTransform.SetScale3D(FVector(1.0f, 1.0f, 1.0f));
 	GizmoTransform.SetRotation(FQuat::Identity);
 	return bHaveSomethingSelected;
@@ -265,34 +354,113 @@ bool USequencerPivotTool::SetGizmoBasedOnSelection(bool bUseSaved)
 void USequencerPivotTool::Shutdown(EToolShutdownType ShutdownType)
 {
 	GizmoManager->DestroyAllGizmosByOwner(this);
-	RemoveControlRigDelegates();
+	RemoveDelegates();
 }
 
-void USequencerPivotTool::HandleControlSelected(UControlRig* Subject, FRigControlElement* InControl, bool bSelected)
+void USequencerPivotTool::DeactivateMe()
 {
-
-	if (FLevelEditorModule* LevelEditorModule = FModuleManager::GetModulePtr<FLevelEditorModule>(TEXT("LevelEditor")))
+	//if in undo make sure we have lost selection else leave alone
+	bool bDeactivate = true;
+	if (GIsTransacting)
 	{
-		TSharedPtr<ILevelEditor> LevelEditorPtr = LevelEditorModule->GetLevelEditorInstance().Pin();
+		bDeactivate = false;
+		for (FControlRigMappings& Mappings : LastSelectedObjects.LastSelectedControlRigs)
+		{
+			if (Mappings.ControlRig.IsValid())
+			{
+				for (TPair<FName, FTransform>& Pair : Mappings.PivotTransforms)
+				{
+
+					if (Mappings.ControlRig->IsControlSelected(Pair.Key) == false)
+					{
+						bDeactivate = true;
+						break;
+					}
+				}
+			}
+			else
+			{
+				bDeactivate = true;
+				break;
+			}
+			
+		}
 
-		if (LevelEditorPtr.IsValid())
+		if(bDeactivate == false)
+		{ 
+			TArray<AActor*> SelectedActors;
+			GEditor->GetSelectedActors()->GetSelectedObjects(SelectedActors);
+
+			for (FActorMappings& ActorMappings : LastSelectedObjects.LastSelectedActors)
+			{
+				
+				if (ActorMappings.Actor.IsValid())
+				{
+					if (SelectedActors.Contains(ActorMappings.Actor.Get()) == false)
+					{
+						bDeactivate = true;
+						break;
+					}
+				}
+				else
+				{
+					bDeactivate = true;
+					break;
+				}
+			}
+		}
+	}
+	if (bDeactivate)
+	{
+		if (FLevelEditorModule* LevelEditorModule = FModuleManager::GetModulePtr<FLevelEditorModule>(TEXT("LevelEditor")))
 		{
-			FString ActiveToolName = LevelEditorPtr->GetEditorModeManager().GetInteractiveToolsContext()->ToolManager->GetActiveToolName(EToolSide::Left);
-			if (ActiveToolName == TEXT("SequencerPivotTool"))
+			TSharedPtr<ILevelEditor> LevelEditorPtr = LevelEditorModule->GetLevelEditorInstance().Pin();
+
+			if (LevelEditorPtr.IsValid())
 			{
-				LevelEditorPtr->GetEditorModeManager().GetInteractiveToolsContext()->ToolManager->DeactivateTool(EToolSide::Left, EToolShutdownType::Completed);
+				FString ActiveToolName = LevelEditorPtr->GetEditorModeManager().GetInteractiveToolsContext()->ToolManager->GetActiveToolName(EToolSide::Left);
+				if (ActiveToolName == TEXT("SequencerPivotTool"))
+				{
+					LevelEditorPtr->GetEditorModeManager().GetInteractiveToolsContext()->ToolManager->DeactivateTool(EToolSide::Left, EToolShutdownType::Completed);
+				}
 			}
 		}
 	}
-	
+
 }
 
-void USequencerPivotTool::RemoveControlRigDelegates()
+void USequencerPivotTool::OnEditorSelectionChanged(UObject* NewSelection)
 {
-	for (UControlRig* ControlRig : ControlRigs)
+	DeactivateMe();
+	/** Later may need to keep track of ordering
+	USelection* Selection = Cast<USelection>(NewSelection);
+	if (!Selection)
 	{
-		ControlRig->ControlSelected().RemoveAll(this);
+		return;
+	}
+
+	TArray<UObject*> SelectedActors;
+	Selection->GetSelectedObjects(AActor::StaticClass(), SelectedActors);
+	*/
+}
+
+void USequencerPivotTool::HandleControlSelected(UControlRig* Subject, FRigControlElement* InControl, bool bSelected)
+{
+	DeactivateMe();
+}
+
+void USequencerPivotTool::RemoveDelegates()
+{
+	for (TWeakObjectPtr<UControlRig>& ControlRig : ControlRigs)
+	{
+		if (ControlRig.IsValid())
+		{
+			ControlRig->ControlSelected().RemoveAll(this);
+		}
 	}
+
+	USelection::SelectionChangedEvent.Remove(OnEditorSelectionChangedHandle);
+	OnEditorSelectionChangedHandle.Reset();
 }
 
 
@@ -308,27 +476,50 @@ void USequencerPivotTool::GizmoTransformStarted(UTransformProxy* Proxy)
 		const FFrameTime FrameTime = Sequencer->GetLocalTime().ConvertTo(TickResolution);
 		const FFrameNumber FrameNumber = FrameTime.RoundToFrame();
 
-		for (UControlRig* ControlRig : ControlRigs)
+		for (TWeakObjectPtr<UControlRig>& ControlRig : ControlRigs)
 		{
+			if (ControlRig.IsValid() == false)
+			{
+				continue;
+			}
 			TArray<FName> SelectedControls = ControlRig->CurrentControlSelection();
 			if (SelectedControls.Num() > 0)
 			{
 				ControlRig->Modify();
 				for (const FName& Name : SelectedControls)
 				{
-					FTransform Transform = UControlRigSequencerEditorLibrary::GetControlRigWorldTransform(LevelSequence, ControlRig, Name, FrameNumber,
+					FTransform Transform = UControlRigSequencerEditorLibrary::GetControlRigWorldTransform(LevelSequence, ControlRig.Get(), Name, FrameNumber,
 						ESequenceTimeUnit::TickResolution);
 
 					FControlRigSelectionDuringDrag ControlDrag;
 					ControlDrag.LevelSequence = LevelSequence;
 					ControlDrag.ControlName = Name;
-					ControlDrag.ControlRig = ControlRig;
+					ControlDrag.ControlRig = ControlRig.Get();
 					ControlDrag.CurrentFrame = FrameNumber;
 					ControlDrag.CurrentTransform = Transform;
 					ControlRigDrags.Add(ControlDrag);
 				}
 			}
 		}
+
+		for (TWeakObjectPtr<AActor>& Actor : Actors)
+		{
+			if (Actor.IsValid() == false)
+			{
+				continue;
+			}
+			//Actor->Modify? Need something for undo.
+
+			FTransform Transform = UControlRigSequencerEditorLibrary::GetActorWorldTransform(LevelSequence, Actor.Get(), FrameNumber,
+				ESequenceTimeUnit::TickResolution);
+
+			FActorSelectonDuringDrag ActorDrag;
+			ActorDrag.LevelSequence = LevelSequence;
+			ActorDrag.Actor = Actor.Get();
+			ActorDrag.CurrentFrame = FrameNumber;
+			ActorDrag.CurrentTransform = Transform;
+			ActorDrags.Add(ActorDrag);
+		}
 	}
 
 	StartDragTransform = Proxy->GetTransform();
@@ -343,23 +534,66 @@ void USequencerPivotTool::GizmoTransformChanged(UTransformProxy* Proxy, FTransfo
 	{
 		return;
 	}
+	bManipulatorMadeChange = true;
+
 	GizmoTransform = Transform;
 	FTransform Diff = Transform.GetRelativeTransform(StartDragTransform);
 	if (Diff.GetRotation().IsIdentity(1e-4f) == false)
 	{
 		const bool bSetKey = false;
-		bManipulatorMadeChange = true;
+		int32 Index = 0;
 		for (FControlRigSelectionDuringDrag& ControlDrag : ControlRigDrags)
 		{
-			FVector LocDiff = ControlDrag.CurrentTransform.GetLocation() - Transform.GetLocation();
-			if (LocDiff.IsNearlyZero(1e-4f) == false)
+			
+			if (!bShiftPressedWhenStarted || Index != (ControlRigDrags.Num() - 1))
+			{
+				FVector LocDiff = ControlDrag.CurrentTransform.GetLocation() - Transform.GetLocation();
+				if (LocDiff.IsNearlyZero(1e-4f) == false)
+				{
+					LocDiff = StartDragTransform.InverseTransformPosition(ControlDrag.CurrentTransform.GetLocation());
+					FVector RotatedDiff = Diff.GetRotation().RotateVector(LocDiff);
+					FVector NewLocation = StartDragTransform.TransformPosition(RotatedDiff);
+					ControlDrag.CurrentTransform.SetLocation(NewLocation);
+					UControlRigSequencerEditorLibrary::SetControlRigWorldTransform(ControlDrag.LevelSequence, ControlDrag.ControlRig, ControlDrag.ControlName,
+						ControlDrag.CurrentFrame, ControlDrag.CurrentTransform, ESequenceTimeUnit::TickResolution, bSetKey);
+				}
+			}
+			else //last one with shift we keep locked!
+			{
+				UControlRigSequencerEditorLibrary::SetControlRigWorldTransform(ControlDrag.LevelSequence, ControlDrag.ControlRig, ControlDrag.ControlName,
+					ControlDrag.CurrentFrame, ControlDrag.CurrentTransform, ESequenceTimeUnit::TickResolution, bSetKey);
+			}
+			
+			++Index;
+		}
+
+		Index = 0;
+		for (FActorSelectonDuringDrag& ActorDrag : ActorDrags)
+		{
+			if (!bShiftPressedWhenStarted || Index != (ActorDrags.Num() - 1))
+			{
+				FVector LocDiff = ActorDrag.CurrentTransform.GetLocation() - Transform.GetLocation();
+				if (LocDiff.IsNearlyZero(1e-4f) == false)
+				{
+					LocDiff = StartDragTransform.InverseTransformPosition(ActorDrag.CurrentTransform.GetLocation());
+					FVector RotatedDiff = Diff.GetRotation().RotateVector(LocDiff);
+					FVector NewLocation = StartDragTransform.TransformPosition(RotatedDiff);
+					ActorDrag.CurrentTransform.SetLocation(NewLocation);
+
+					if (USceneComponent* RootComponent = ActorDrag.Actor->GetRootComponent())
+					{
+						RootComponent->SetWorldLocation(NewLocation);
+					}
+				}
+			}
+			else //last one with shift we keep locked!
 			{
-				FVector RotatedDiff = Diff.GetRotation().RotateVector(LocDiff);
-				FVector NewLocation = Transform.GetLocation() + RotatedDiff;
-				ControlDrag.CurrentTransform.SetLocation(NewLocation);
-				UControlRigSequencerEditorLibrary::SetControlRigWorldTransform(ControlDrag.LevelSequence, ControlDrag.ControlRig, ControlDrag.ControlName, 
-					ControlDrag.CurrentFrame,ControlDrag.CurrentTransform, ESequenceTimeUnit::TickResolution, bSetKey);
+				if (USceneComponent* RootComponent = ActorDrag.Actor->GetRootComponent())
+				{
+					RootComponent->SetWorldLocation(ActorDrag.CurrentTransform.GetLocation());
+				}
 			}
+			++Index;
 		}
 	}
 	
@@ -387,9 +621,9 @@ void USequencerPivotTool::GizmoTransformEnded(UTransformProxy* Proxy)
 
 void USequencerPivotTool::SavePivotTransforms()
 {
-	for (FControlRigMappings& LastObject : LastSelectedObjects)
+	for (FControlRigMappings& LastObject : LastSelectedObjects.LastSelectedControlRigs)
 	{
-		FControlRigMappings& Mappings = SavedPivotLocations.FindOrAdd(LastObject.ControlRig);
+		FControlRigMappings& Mappings = SavedPivotLocations.ControlRigMappings.FindOrAdd(LastObject.ControlRig);
 		Mappings.ControlRig = LastObject.ControlRig;
 		for (TPair<FName, FTransform>& Pair : LastObject.PivotTransforms)
 		{
@@ -397,6 +631,13 @@ void USequencerPivotTool::SavePivotTransforms()
 			Transform = GizmoTransform;
 		}
 	}
+
+	for (FActorMappings& LastObject : LastSelectedObjects.LastSelectedActors)
+	{
+		FActorMappings& Mappings = SavedPivotLocations.ActorMappings.FindOrAdd(LastObject.Actor);
+		Mappings.Actor = LastObject.Actor;
+		Mappings.PivotTransform = GizmoTransform;
+	}
 }
 
 //currently always visibile
diff --git a/Engine/Plugins/Editor/SequencerAnimTools/Source/SequencerAnimTools/Private/MotionTrailTool.cpp b/Engine/Plugins/Editor/SequencerAnimTools/Source/SequencerAnimTools/Private/MotionTrailTool.cpp
--- a/Engine/Plugins/Editor/SequencerAnimTools/Source/SequencerAnimTools/Private/MotionTrailTool.cpp
+++ b/Engine/Plugins/Editor/SequencerAnimTools/Source/SequencerAnimTools/Private/MotionTrailTool.cpp
@@ -90,7 +90,7 @@ void UMotionTrailTool::Setup()
 	//TransformGizmo = GetToolManager()->GetPairedGizmoManager()->Create3AxisTransformGizmo(this, UMotionTrailTool::TrailKeyTransformGizmoInstanceIdentifier);
 
 	FString GizmoIdentifier = TEXT("PivotToolGizmoIdentifier");
-	ETransformGizmoSubElements Elements = ETransformGizmoSubElements::TranslateAllAxes;
+	ETransformGizmoSubElements Elements = ETransformGizmoSubElements::StandardTranslateRotate;
 	TransformGizmo = GizmoManager->CreateCustomTransformGizmo(Elements, this, UMotionTrailTool::TrailKeyTransformGizmoInstanceIdentifier);
 
 	TransformProxy->OnTransformChanged.AddUObject(this, &UMotionTrailTool::GizmoTransformChanged);
@@ -293,11 +293,12 @@ void UMotionTrailTool::GizmoTransformChanged(UTransformProxy* Proxy, FTransform
 	}
 	GizmoTransform = Transform;
 	FTransform Diff = Transform.GetRelativeTransform(StartDragTransform);
+	FVector LocationDiff = GizmoTransform.GetLocation() - StartDragTransform.GetLocation();
 	for (const TUniquePtr<UE::SequencerAnimTools::FTrailHierarchy>& TrailHierarchy : TrailHierarchies)
 	{
 		FVector WidgetLocation = Transform.GetLocation();
 		FVector Location;
-		bool bIsHandled = TrailHierarchy->ApplyDelta(Diff.GetTranslation(), Diff.GetRotation().Rotator(), WidgetLocation);
+		bool bIsHandled = TrailHierarchy->ApplyDelta(LocationDiff, Diff.GetRotation().Rotator(), WidgetLocation);
 		if (bIsHandled)
 		{
 			bManipulatorMadeChange = true;
diff --git a/Engine/Plugins/Editor/SequencerAnimTools/Source/SequencerAnimTools/Private/Sequencer/SequencerTrailHierarchy.cpp b/Engine/Plugins/Editor/SequencerAnimTools/Source/SequencerAnimTools/Private/Sequencer/SequencerTrailHierarchy.cpp
--- a/Engine/Plugins/Editor/SequencerAnimTools/Source/SequencerAnimTools/Private/Sequencer/SequencerTrailHierarchy.cpp
+++ b/Engine/Plugins/Editor/SequencerAnimTools/Source/SequencerAnimTools/Private/Sequencer/SequencerTrailHierarchy.cpp
@@ -743,6 +743,7 @@ void FSequencerTrailHierarchy::RegisterControlRigDelegates(USkeletalMeshComponen
 			}
 
 			//check to see if the seleced control rig is sill selected
+			TArray<FGuid> TrailsToRemove;
 			for (TPair<UControlRig*, TMap<FName, FGuid>>& CompMapPair : ControlsTracked)
 			{
 				UControlRig* TrackedControlRig = CompMapPair.Key;
@@ -750,12 +751,16 @@ void FSequencerTrailHierarchy::RegisterControlRigDelegates(USkeletalMeshComponen
 				{
 					if (TrackedControlRig->IsControlSelected(NameGuid.Key) == false)
 					{
-						const FGuid TrailGuid = NameGuid.Value;
-						VisibilityManager.ControlSelected.Remove(TrailGuid);
-						RemoveTrailIfNotAlwaysVisible(TrailGuid);
+						FGuid TrailGuid = NameGuid.Value;
+						TrailsToRemove.Add(TrailGuid);
 					}
 				}
 			}
+			for (const FGuid TrailGuid: TrailsToRemove)
+			{
+				VisibilityManager.ControlSelected.Remove(TrailGuid);
+				RemoveTrailIfNotAlwaysVisible(TrailGuid);
+			}
 		}
 
 			
diff --git a/Engine/Plugins/Editor/SequencerAnimTools/Source/SequencerAnimTools/Public/EditPivotTool.h b/Engine/Plugins/Editor/SequencerAnimTools/Source/SequencerAnimTools/Public/EditPivotTool.h
--- a/Engine/Plugins/Editor/SequencerAnimTools/Source/SequencerAnimTools/Public/EditPivotTool.h
+++ b/Engine/Plugins/Editor/SequencerAnimTools/Source/SequencerAnimTools/Public/EditPivotTool.h
@@ -42,13 +42,22 @@ public:
 	virtual UInteractiveTool* BuildTool(const FToolBuilderState& SceneState) const override;
 };
 
-struct FControlRigSelectionDuringDrag
+struct FSelectionDuringDrag
 {
 	ULevelSequence* LevelSequence;
 	FFrameNumber CurrentFrame;
+	FTransform CurrentTransform;
+
+};
+struct FControlRigSelectionDuringDrag : public FSelectionDuringDrag
+{
 	UControlRig* ControlRig;
 	FName ControlName;
-	FTransform CurrentTransform;
+};
+
+struct FActorSelectonDuringDrag : public FSelectionDuringDrag
+{
+	AActor* Actor;
 };
 
 struct FControlRigMappings
@@ -57,6 +66,24 @@ struct FControlRigMappings
 	TMap<FName, FTransform> PivotTransforms;
 };
 
+struct FActorMappings
+{
+	TWeakObjectPtr<AActor> Actor;
+	FTransform PivotTransform;
+};
+
+struct FSavedMappings
+{
+	TMap<TWeakObjectPtr<UControlRig>, FControlRigMappings> ControlRigMappings;
+	TMap<TWeakObjectPtr<AActor>, FActorMappings> ActorMappings;
+};
+
+struct FLastSelectedObjects
+{
+	TArray<FControlRigMappings> LastSelectedControlRigs;
+	TArray<FActorMappings> LastSelectedActors;
+
+};
 class FEditPivotCommands : public TCommands<FEditPivotCommands>
 {
 public:
@@ -83,8 +110,9 @@ public:
 	TMap<FName, TArray<TSharedPtr<FUICommandInfo>>> Commands;
 };
 
-/**
 
+/**
+Pivot tool class
  */
 UCLASS()
 class  USequencerPivotTool : public UMultiSelectionTool, public IClickBehaviorTarget  , public IBaseSequencerAnimTool
@@ -125,6 +153,7 @@ protected:
 
 protected:
 	bool bShiftPressedWhenStarted = false;
+	bool bCtrlPressedWhenStarted = false;
 	int32 CtrlModifierId = 1;
 	UWorld* TargetWorld = nullptr;		// target World we will raycast into
 	UInteractiveGizmoManager* GizmoManager = nullptr; //gizmo man
@@ -134,11 +163,13 @@ protected:
 	bool bManipulatorMadeChange = false;
 	int32 TransactionIndex = -1;
 	TArray<FControlRigSelectionDuringDrag> ControlRigDrags;
+	TArray<FActorSelectonDuringDrag> ActorDrags;
 
 	//since we are selection based we can cache this
 	ULevelSequence* LevelSequence;
-	TArray<UControlRig*> ControlRigs;
+	TArray<TWeakObjectPtr<UControlRig>> ControlRigs;
 	TWeakPtr<ISequencer> SequencerPtr;
+	TArray<TWeakObjectPtr<AActor>> Actors;
 
 	FTransform GizmoTransform = FTransform::Identity;
 	bool bPickingPivotLocation = false; //mauy remove
@@ -157,16 +188,21 @@ protected:
 	//Handle Selection and Pivot Location
 	void SavePivotTransforms();
 	void SaveLastSelected();
-	void RemoveControlRigDelegates();
+
+	// selection delegates
+	void DeactivateMe();
+	void RemoveDelegates();
 	void HandleControlSelected(UControlRig* Subject, FRigControlElement* InControl, bool bSelected);
+	void OnEditorSelectionChanged(UObject* NewSelection);
+	FDelegateHandle OnEditorSelectionChangedHandle;
 
 private:
 	TSharedPtr<FUICommandList> CommandBindings;
 	void ResetPivot();
 
 public:
-	static TMap<TWeakObjectPtr<UControlRig>, FControlRigMappings> SavedPivotLocations;
-	static TArray<FControlRigMappings> LastSelectedObjects;
+	static FSavedMappings SavedPivotLocations;
+	static FLastSelectedObjects LastSelectedObjects;
 };
 
 
