Sequencer: Control Rig: Enable AnimTools by default (since CR is default now).

#trivial

diff --git a/Engine/Plugins/Editor/SequencerAnimTools/SequencerAnimTools.uplugin b/Engine/Plugins/Editor/SequencerAnimTools/SequencerAnimTools.uplugin
--- a/Engine/Plugins/Editor/SequencerAnimTools/SequencerAnimTools.uplugin
+++ b/Engine/Plugins/Editor/SequencerAnimTools/SequencerAnimTools.uplugin
@@ -14,6 +14,7 @@
 	"IsBetaVersion": false,
 	"IsExperimentalVersion": false,
 	"Installed": false,
+	"EnabledByDefault": true,
 	"Plugins": [
 		{
 			"Name": "ControlRig",
