Sequencer: Fallout from float to double channel, control rig still uses float channels so put support back in for motion trail.
#trivial

diff --git a/Engine/Plugins/Editor/SequencerAnimTools/Source/SequencerAnimTools/Private/Sequencer/MotionTrailMovieSceneKey.cpp b/Engine/Plugins/Editor/SequencerAnimTools/Source/SequencerAnimTools/Private/Sequencer/MotionTrailMovieSceneKey.cpp
--- a/Engine/Plugins/Editor/SequencerAnimTools/Source/SequencerAnimTools/Private/Sequencer/MotionTrailMovieSceneKey.cpp
+++ b/Engine/Plugins/Editor/SequencerAnimTools/Source/SequencerAnimTools/Private/Sequencer/MotionTrailMovieSceneKey.cpp
@@ -195,18 +195,40 @@ void FMotionTraiMovieScenelKeyTool::BuildKeys()
 
 	UMovieSceneSection* AbsoluteTransformSection = OwningTrail->GetSection();
 	TArrayView<FMovieSceneDoubleChannel*> DoubleChannels = AbsoluteTransformSection->GetChannelProxy().GetChannels<FMovieSceneDoubleChannel>();
-	DoubleChannels = DoubleChannels.Slice(OwningTrail->GetChannelOffset(), uint8(ETransformChannel::MaxChannel) + 1);
-	for (uint8 ChannelIdx = 0; ChannelIdx <= uint8(ETransformChannel::MaxChannel); ChannelIdx++)
+	if (DoubleChannels.Num() > 0)
 	{
-		FMovieSceneDoubleChannel* DoubleChannel = DoubleChannels[ChannelIdx];
-		for (int32 Idx = 0; Idx < DoubleChannel->GetNumKeys(); Idx++)
+		DoubleChannels = DoubleChannels.Slice(OwningTrail->GetChannelOffset(), uint8(ETransformChannel::MaxChannel) + 1);
+		for (uint8 ChannelIdx = 0; ChannelIdx <= uint8(ETransformChannel::MaxChannel); ChannelIdx++)
 		{
-			FFrameNumber CurTime = DoubleChannel->GetTimes()[Idx];
+			FMovieSceneDoubleChannel* DoubleChannel = DoubleChannels[ChannelIdx];
+			for (int32 Idx = 0; Idx < DoubleChannel->GetNumKeys(); Idx++)
+			{
+				FFrameNumber CurTime = DoubleChannel->GetTimes()[Idx];
 
-			if (!Keys.Contains(CurTime))
+				if (!Keys.Contains(CurTime))
+				{
+					TUniquePtr<FTrailKeyInfo> TempKeyInfo = MakeUnique<FTrailKeyInfo>(CurTime, AbsoluteTransformSection, OwningTrail);
+					Keys.Add(MoveTemp(CurTime), MoveTemp(TempKeyInfo));
+				}
+			}
+		}
+	}
+	TArrayView<FMovieSceneFloatChannel*> FloatChannels = AbsoluteTransformSection->GetChannelProxy().GetChannels<FMovieSceneFloatChannel>();
+	if (FloatChannels.Num() > 0)
+	{
+		FloatChannels = FloatChannels.Slice(OwningTrail->GetChannelOffset(), uint8(ETransformChannel::MaxChannel) + 1);
+		for (uint8 ChannelIdx = 0; ChannelIdx <= uint8(ETransformChannel::MaxChannel); ChannelIdx++)
+		{
+			FMovieSceneFloatChannel* DoubleChannel = FloatChannels[ChannelIdx];
+			for (int32 Idx = 0; Idx < DoubleChannel->GetNumKeys(); Idx++)
 			{
-				TUniquePtr<FTrailKeyInfo> TempKeyInfo = MakeUnique<FTrailKeyInfo>(CurTime, AbsoluteTransformSection, OwningTrail);
-				Keys.Add(MoveTemp(CurTime), MoveTemp(TempKeyInfo));
+				FFrameNumber CurTime = DoubleChannel->GetTimes()[Idx];
+
+				if (!Keys.Contains(CurTime))
+				{
+					TUniquePtr<FTrailKeyInfo> TempKeyInfo = MakeUnique<FTrailKeyInfo>(CurTime, AbsoluteTransformSection, OwningTrail);
+					Keys.Add(MoveTemp(CurTime), MoveTemp(TempKeyInfo));
+				}
 			}
 		}
 	}
@@ -243,16 +265,35 @@ void FMotionTraiMovieScenelKeyTool::TranslateSelectedKeys(bool bRight)
 		FFrameNumber Delta = FQualifiedFrameTime(Shift, OwningTrail->GetSequencer()->GetFocusedDisplayRate()).ConvertTo(OwningTrail->GetSequencer()->GetFocusedTickResolution()).RoundToFrame();
 		TRange<FFrameNumber> SectionNewBounds = Section->GetRange();
 		TArrayView<FMovieSceneDoubleChannel*> DoubleChannels = Section->GetChannelProxy().GetChannels<FMovieSceneDoubleChannel>();
-		DoubleChannels = DoubleChannels.Slice(OwningTrail->GetChannelOffset(), uint8(ETransformChannel::MaxChannel) + 1);
+		bool bIsFloat = true;
+		if (DoubleChannels.Num() > 0)
+		{
+			bIsFloat = false;
+			DoubleChannels = DoubleChannels.Slice(OwningTrail->GetChannelOffset(), uint8(ETransformChannel::MaxChannel) + 1);
+		}
+		TArrayView<FMovieSceneFloatChannel*> FloatChannels = Section->GetChannelProxy().GetChannels<FMovieSceneFloatChannel>();
+		if (FloatChannels.Num() > 0)
+		{
+			bIsFloat = true;
+			FloatChannels = FloatChannels.Slice(OwningTrail->GetChannelOffset(), uint8(ETransformChannel::MaxChannel) + 1);
+		}
 		for (uint8 ChannelIdx = 0; ChannelIdx <= uint8(ETransformChannel::MaxChannel); ChannelIdx++)
 		{
 			TArray<FFrameNumber> KeyTimes;
-			FMovieSceneDoubleChannel* DoubleChannel = DoubleChannels[ChannelIdx];
-			TArray<FKeyHandle> KeyHandles = GetSelectedKeyHandles(DoubleChannel);
+			FMovieSceneChannel* Channel = nullptr;
+			if (bIsFloat)
+			{
+				Channel = FloatChannels[ChannelIdx];
+			}
+			else
+			{
+				Channel = DoubleChannels[ChannelIdx];
+			}
+			TArray<FKeyHandle> KeyHandles = GetSelectedKeyHandles(Channel);
 			if (KeyHandles.Num() > 0)
 			{
 				KeyTimes.SetNum(KeyHandles.Num());
-				DoubleChannel->GetKeyTimes(KeyHandles, KeyTimes);
+				Channel->GetKeyTimes(KeyHandles, KeyTimes);
 				FFrameNumber LowestFrameTime = KeyTimes[0];
 				FFrameNumber HighestFrameTime = KeyTimes[0];
 
@@ -272,7 +313,7 @@ void FMotionTraiMovieScenelKeyTool::TranslateSelectedKeys(bool bRight)
 				}
 				SectionNewBounds = TRange<FFrameNumber>::Hull(SectionNewBounds, TRange<FFrameNumber>(LowestFrameTime, HighestFrameTime + 1));
 				Section->SetRange(SectionNewBounds);
-				DoubleChannel->SetKeyTimes(KeyHandles, KeyTimes);
+				Channel->SetKeyTimes(KeyHandles, KeyTimes);
 			}
 
 		}
@@ -294,13 +335,32 @@ void FMotionTraiMovieScenelKeyTool::DeleteSelectedKeys()
 	if (CachedSelection.Num() > 0 && Section && Section->TryModify())
 	{
 		Section->Modify();
+		bool bIsFloat = true;
 		TArrayView<FMovieSceneDoubleChannel*> DoubleChannels = Section->GetChannelProxy().GetChannels<FMovieSceneDoubleChannel>();
-		DoubleChannels = DoubleChannels.Slice(OwningTrail->GetChannelOffset(), uint8(ETransformChannel::MaxChannel) + 1);
+		if (DoubleChannels.Num() > 0)
+		{
+			DoubleChannels = DoubleChannels.Slice(OwningTrail->GetChannelOffset(), uint8(ETransformChannel::MaxChannel) + 1);
+			bIsFloat = false;
+		}
+		TArrayView<FMovieSceneFloatChannel*> FloatChannels = Section->GetChannelProxy().GetChannels<FMovieSceneFloatChannel>();
+		if (FloatChannels.Num() > 0)
+		{
+			FloatChannels = FloatChannels.Slice(OwningTrail->GetChannelOffset(), uint8(ETransformChannel::MaxChannel) + 1);
+			bIsFloat = true;
+		}
 		for (uint8 ChannelIdx = 0; ChannelIdx <= uint8(ETransformChannel::MaxChannel); ChannelIdx++)
 		{
-			FMovieSceneDoubleChannel* DoubleChannel = DoubleChannels[ChannelIdx];
-			TArray<FKeyHandle> KeyHandles = GetSelectedKeyHandles(DoubleChannel);
-			DoubleChannel->DeleteKeys(KeyHandles);
+			FMovieSceneChannel* Channel = nullptr;
+			if (bIsFloat)
+			{
+				Channel = FloatChannels[ChannelIdx];
+			}
+			else
+			{
+				Channel = DoubleChannels[ChannelIdx];
+			}
+			TArray<FKeyHandle> KeyHandles = GetSelectedKeyHandles(Channel);
+			Channel->DeleteKeys(KeyHandles);
 		}
 		OwningTrail->GetSequencer()->NotifyMovieSceneDataChanged(EMovieSceneDataChangeType::TrackValueChanged);
 		CachedSelection.Reset();
@@ -433,15 +493,31 @@ bool FMotionTraiMovieScenelKeyTool::ShouldRebuildKeys()
 {
 	TMap<FFrameNumber, TSet<ETransformChannel>> KeyTimes;
 	TArrayView<FMovieSceneDoubleChannel*> DoubleChannels = OwningTrail->GetSection()->GetChannelProxy().GetChannels<FMovieSceneDoubleChannel>();
-	DoubleChannels = DoubleChannels.Slice(OwningTrail->GetChannelOffset(), uint8(ETransformChannel::MaxChannel) + 1);
-	for (uint8 ChannelIdx = 0; ChannelIdx <= uint8(ETransformChannel::MaxChannel); ChannelIdx++)
-
+	if (DoubleChannels.Num() > 0)
 	{
-		FMovieSceneDoubleChannel* DoubleChannel = DoubleChannels[ChannelIdx];
-		for (int32 Idx = 0; Idx < DoubleChannel->GetNumKeys(); Idx++)
+		DoubleChannels = DoubleChannels.Slice(OwningTrail->GetChannelOffset(), uint8(ETransformChannel::MaxChannel) + 1);
+		for (uint8 ChannelIdx = 0; ChannelIdx <= uint8(ETransformChannel::MaxChannel); ChannelIdx++)
 		{
-			const FFrameNumber CurTime = DoubleChannel->GetTimes()[Idx];
-			KeyTimes.FindOrAdd(CurTime).Add(ETransformChannel(ChannelIdx));
+			FMovieSceneDoubleChannel* DoubleChannel = DoubleChannels[ChannelIdx];
+			for (int32 Idx = 0; Idx < DoubleChannel->GetNumKeys(); Idx++)
+			{
+				const FFrameNumber CurTime = DoubleChannel->GetTimes()[Idx];
+				KeyTimes.FindOrAdd(CurTime).Add(ETransformChannel(ChannelIdx));
+			}
+		}
+	}
+	TArrayView<FMovieSceneFloatChannel*> FloatChannels = OwningTrail->GetSection()->GetChannelProxy().GetChannels<FMovieSceneFloatChannel>();
+	if (FloatChannels.Num() > 0)
+	{
+		FloatChannels = FloatChannels.Slice(OwningTrail->GetChannelOffset(), uint8(ETransformChannel::MaxChannel) + 1);
+		for (uint8 ChannelIdx = 0; ChannelIdx <= uint8(ETransformChannel::MaxChannel); ChannelIdx++)
+		{
+			FMovieSceneFloatChannel* FloatChannel = FloatChannels[ChannelIdx];
+			for (int32 Idx = 0; Idx < FloatChannel->GetNumKeys(); Idx++)
+			{
+				const FFrameNumber CurTime = FloatChannel->GetTimes()[Idx];
+				KeyTimes.FindOrAdd(CurTime).Add(ETransformChannel(ChannelIdx));
+			}
 		}
 	}
 
@@ -460,7 +536,10 @@ bool FMotionTraiMovieScenelKeyTool::ShouldRebuildKeys()
 		{
 			const ETransformChannel TransformChannel = ETransformChannel(ChannelIdx);
 			if ((!TimeKeyPair.Value.Contains(TransformChannel) && Keys[TimeKeyPair.Key]->IdxMap.Contains(TransformChannel)) ||
-				(Keys[TimeKeyPair.Key]->IdxMap.Contains(TransformChannel) && DoubleChannels[uint8(TransformChannel)]->GetData().GetIndex(Keys[TimeKeyPair.Key]->IdxMap[TransformChannel]) == INDEX_NONE) ||
+				(Keys[TimeKeyPair.Key]->IdxMap.Contains(TransformChannel) &&
+					(DoubleChannels.Num() > 0 && DoubleChannels[uint8(TransformChannel)]->GetData().GetIndex(Keys[TimeKeyPair.Key]->IdxMap[TransformChannel]) == INDEX_NONE) ||
+					(FloatChannels.Num() > 0 && FloatChannels[uint8(TransformChannel)]->GetData().GetIndex(Keys[TimeKeyPair.Key]->IdxMap[TransformChannel]) == INDEX_NONE)
+					) ||
 				(TimeKeyPair.Value.Contains(TransformChannel) && !Keys[TimeKeyPair.Key]->IdxMap.Contains(TransformChannel)))
 			{
 				return true;
@@ -501,14 +580,34 @@ FTrailKeyInfo::FTrailKeyInfo(const FFrameNumber InFrameNumber, UMovieSceneSectio
 , bDirty(true)
 , OwningTrail(InOwningTrail)
 {
-	TArrayView<FMovieSceneDoubleChannel*> Channels = InSection->GetChannelProxy().GetChannels<FMovieSceneDoubleChannel>();
-	Channels = Channels.Slice(OwningTrail->GetChannelOffset(), uint8(ETransformChannel::MaxChannel) + 1);
-	for (uint8 Idx = 0; Idx <= uint8(ETransformChannel::MaxChannel); Idx++)
 	{
-		const int32 FoundIdx = Channels[Idx]->GetData().FindKey(InFrameNumber);
-		if (FoundIdx != INDEX_NONE)
+		TArrayView<FMovieSceneDoubleChannel*> Channels = InSection->GetChannelProxy().GetChannels<FMovieSceneDoubleChannel>();
+		if (Channels.Num() > 0)
 		{
-			IdxMap.Add(ETransformChannel(Idx), Channels[Idx]->GetData().GetHandle(FoundIdx));
+			Channels = Channels.Slice(OwningTrail->GetChannelOffset(), uint8(ETransformChannel::MaxChannel) + 1);
+			for (uint8 Idx = 0; Idx <= uint8(ETransformChannel::MaxChannel); Idx++)
+			{
+				const int32 FoundIdx = Channels[Idx]->GetData().FindKey(InFrameNumber);
+				if (FoundIdx != INDEX_NONE)
+				{
+					IdxMap.Add(ETransformChannel(Idx), Channels[Idx]->GetData().GetHandle(FoundIdx));
+				}
+			}
+		}
+	}
+	{
+		TArrayView<FMovieSceneFloatChannel*> Channels = InSection->GetChannelProxy().GetChannels<FMovieSceneFloatChannel>();
+		if (Channels.Num() > 0)
+		{
+			Channels = Channels.Slice(OwningTrail->GetChannelOffset(), uint8(ETransformChannel::MaxChannel) + 1);
+			for (uint8 Idx = 0; Idx <= uint8(ETransformChannel::MaxChannel); Idx++)
+			{
+				const int32 FoundIdx = Channels[Idx]->GetData().FindKey(InFrameNumber);
+				if (FoundIdx != INDEX_NONE)
+				{
+					IdxMap.Add(ETransformChannel(Idx), Channels[Idx]->GetData().GetHandle(FoundIdx));
+				}
+			}
 		}
 	}
 }
diff --git a/Engine/Plugins/Editor/SequencerAnimTools/Source/SequencerAnimTools/Private/Sequencer/MovieSceneTransformTrail.cpp b/Engine/Plugins/Editor/SequencerAnimTools/Source/SequencerAnimTools/Private/Sequencer/MovieSceneTransformTrail.cpp
--- a/Engine/Plugins/Editor/SequencerAnimTools/Source/SequencerAnimTools/Private/Sequencer/MovieSceneTransformTrail.cpp
+++ b/Engine/Plugins/Editor/SequencerAnimTools/Source/SequencerAnimTools/Private/Sequencer/MovieSceneTransformTrail.cpp
@@ -283,11 +283,27 @@ TRange<double> FMovieSceneTransformTrail::GetEffectiveSectionRange(int32 Channel
 
 	TRange<FFrameNumber> EffectiveRange = TRange<FFrameNumber>::Empty();
 
-	TArrayView<FMovieSceneDoubleChannel*> Channels = TransformSection->GetChannelProxy().GetChannels<FMovieSceneDoubleChannel>();
-	for (int32 ChannelIdx = ChannelOffset; ChannelIdx <= ChannelOffset + uint8(ETransformChannel::MaxChannel); ChannelIdx++)
 	{
-		FMovieSceneDoubleChannel* Channel = Channels[ChannelIdx];
-		EffectiveRange = TRange<FFrameNumber>::Hull(EffectiveRange, Channel->ComputeEffectiveRange());
+		TArrayView<FMovieSceneDoubleChannel*> Channels = TransformSection->GetChannelProxy().GetChannels<FMovieSceneDoubleChannel>();
+		if (Channels.Num() > 0)
+		{
+			for (int32 ChannelIdx = ChannelOffset; ChannelIdx <= ChannelOffset + uint8(ETransformChannel::MaxChannel); ChannelIdx++)
+			{
+				FMovieSceneDoubleChannel* Channel = Channels[ChannelIdx];
+				EffectiveRange = TRange<FFrameNumber>::Hull(EffectiveRange, Channel->ComputeEffectiveRange());
+			}
+		}
+	}
+	{
+		TArrayView<FMovieSceneFloatChannel*> Channels = TransformSection->GetChannelProxy().GetChannels<FMovieSceneFloatChannel>();
+		if (Channels.Num() > 0)
+		{
+			for (int32 ChannelIdx = ChannelOffset; ChannelIdx <= ChannelOffset + uint8(ETransformChannel::MaxChannel); ChannelIdx++)
+			{
+				FMovieSceneFloatChannel* Channel = Channels[ChannelIdx];
+				EffectiveRange = TRange<FFrameNumber>::Hull(EffectiveRange, Channel->ComputeEffectiveRange());
+			}
+		}
 	}
 
 	EffectiveRange = TRange<FFrameNumber>::Intersection(EffectiveRange, TransformSection->GetRange());
diff --git a/Engine/Plugins/Editor/SequencerAnimTools/Source/SequencerAnimTools/Private/Sequencer/SequencerTrailHierarchy.cpp b/Engine/Plugins/Editor/SequencerAnimTools/Source/SequencerAnimTools/Private/Sequencer/SequencerTrailHierarchy.cpp
--- a/Engine/Plugins/Editor/SequencerAnimTools/Source/SequencerAnimTools/Private/Sequencer/SequencerTrailHierarchy.cpp
+++ b/Engine/Plugins/Editor/SequencerAnimTools/Source/SequencerAnimTools/Private/Sequencer/SequencerTrailHierarchy.cpp
@@ -62,13 +62,9 @@ void FSequencerTrailHierarchy::Initialize()
 		VisibilityManager.Selected.Add(Guid);
 	});
 	
-	OnSelectionChangedHandle = Sequencer->GetSelectionChangedObjectGuids().AddLambda([this](TArray<FGuid> NewSelection) {
-		
-		//if using old trails don't add new ones.
-		if (CVarUseOldSequencerMotionTrails->GetBool() == true)
-		{
-			return;
-		}		
+	OnSelectionChangedHandle = Sequencer->GetSelectionChangedObjectGuids().AddLambda([this](TArray<FGuid> NewSelection)
+	{
+	
 		TSharedPtr<ISequencer> Sequencer = WeakSequencer.Pin();
 		check(Sequencer);
 		TSet<FGuid> OldSelected = VisibilityManager.Selected;
@@ -451,11 +447,6 @@ void FSequencerTrailHierarchy::OnBindingVisibilityStateChanged(UObject* BoundObj
 
 void FSequencerTrailHierarchy::UpdateSequencerBindings(const TArray<FGuid>& SequencerBindings, TFunctionRef<void(UObject*, FTrail*, FGuid)> OnUpdated)
 {
-	//if using old trails don't add new ones.
-	if (CVarUseOldSequencerMotionTrails->GetBool() == true)
-	{
-		return;
-	}
 	const FDateTime StartTime = FDateTime::Now();
 
 	TSharedPtr<ISequencer> Sequencer = WeakSequencer.Pin();
@@ -471,6 +462,11 @@ void FSequencerTrailHierarchy::UpdateSequencerBindings(const TArray<FGuid>& Sequ
 				{
 					continue;
 				}
+				//if using old trails don't add new ones.
+				if (CVarUseOldSequencerMotionTrails->GetBool() == true)
+				{
+					continue;
+				}
 
 				USceneComponent* BoundComponent = Cast<USceneComponent>(BoundObject.Get());
 				if (AActor* BoundActor = Cast<AActor>(BoundObject.Get()))
@@ -502,6 +498,10 @@ void FSequencerTrailHierarchy::UpdateSequencerBindings(const TArray<FGuid>& Sequ
 				{
 					continue;
 				}
+				if (CVarUseOldSequencerMotionTrails->GetBool() == true)
+				{
+					continue;
+				}
 
 				USkeletalMeshComponent* BoundComponent = Cast<USkeletalMeshComponent>(BoundObject.Get());
 				if (AActor* BoundActor = Cast<AActor>(BoundObject.Get()))
