CIS: Linux incremental Fix

diff --git a/Engine/Plugins/Editor/SequencerAnimTools/Source/SequencerAnimTools/Private/Sequencer/MotionTrailMovieSceneKey.cpp b/Engine/Plugins/Editor/SequencerAnimTools/Source/SequencerAnimTools/Private/Sequencer/MotionTrailMovieSceneKey.cpp
--- a/Engine/Plugins/Editor/SequencerAnimTools/Source/SequencerAnimTools/Private/Sequencer/MotionTrailMovieSceneKey.cpp
+++ b/Engine/Plugins/Editor/SequencerAnimTools/Source/SequencerAnimTools/Private/Sequencer/MotionTrailMovieSceneKey.cpp
@@ -535,14 +535,23 @@ bool FMotionTraiMovieScenelKeyTool::ShouldRebuildKeys()
 		for (uint8 ChannelIdx = 0; ChannelIdx <= uint8(ETransformChannel::MaxChannel); ChannelIdx++)
 		{
 			const ETransformChannel TransformChannel = ETransformChannel(ChannelIdx);
-			if ((!TimeKeyPair.Value.Contains(TransformChannel) && Keys[TimeKeyPair.Key]->IdxMap.Contains(TransformChannel)) ||
-				(Keys[TimeKeyPair.Key]->IdxMap.Contains(TransformChannel) &&
-					(DoubleChannels.Num() > 0 && DoubleChannels[uint8(TransformChannel)]->GetData().GetIndex(Keys[TimeKeyPair.Key]->IdxMap[TransformChannel]) == INDEX_NONE) ||
-					(FloatChannels.Num() > 0 && FloatChannels[uint8(TransformChannel)]->GetData().GetIndex(Keys[TimeKeyPair.Key]->IdxMap[TransformChannel]) == INDEX_NONE)
-					) ||
-				(TimeKeyPair.Value.Contains(TransformChannel) && !Keys[TimeKeyPair.Key]->IdxMap.Contains(TransformChannel)))
+			if (DoubleChannels.Num() > 0)
 			{
-				return true;
+				if ((!TimeKeyPair.Value.Contains(TransformChannel) && Keys[TimeKeyPair.Key]->IdxMap.Contains(TransformChannel)) ||
+					(Keys[TimeKeyPair.Key]->IdxMap.Contains(TransformChannel) && DoubleChannels[uint8(TransformChannel)]->GetData().GetIndex(Keys[TimeKeyPair.Key]->IdxMap[TransformChannel]) == INDEX_NONE) ||
+					(TimeKeyPair.Value.Contains(TransformChannel) && !Keys[TimeKeyPair.Key]->IdxMap.Contains(TransformChannel)))
+				{
+					return true;
+				}
+			}
+			if (FloatChannels.Num() > 0)
+			{
+				if ((!TimeKeyPair.Value.Contains(TransformChannel) && Keys[TimeKeyPair.Key]->IdxMap.Contains(TransformChannel)) ||
+					(Keys[TimeKeyPair.Key]->IdxMap.Contains(TransformChannel) && FloatChannels[uint8(TransformChannel)]->GetData().GetIndex(Keys[TimeKeyPair.Key]->IdxMap[TransformChannel]) == INDEX_NONE) ||
+					(TimeKeyPair.Value.Contains(TransformChannel) && !Keys[TimeKeyPair.Key]->IdxMap.Contains(TransformChannel)))
+				{
+					return true;
+				}
 			}
 		}
 	}
