Sequencer: Fix tool registration/unregistration so that they happen in EditMode::Enter and Exit. Otherwise, the tools weren't getting unregistered, leading to ensures when registering again.

diff --git a/Engine/Plugins/Editor/SequencerAnimTools/Source/SequencerAnimTools/Private/SequencerAnimToolsModule.cpp b/Engine/Plugins/Editor/SequencerAnimTools/Source/SequencerAnimTools/Private/SequencerAnimToolsModule.cpp
--- a/Engine/Plugins/Editor/SequencerAnimTools/Source/SequencerAnimTools/Private/SequencerAnimToolsModule.cpp
+++ b/Engine/Plugins/Editor/SequencerAnimTools/Source/SequencerAnimTools/Private/SequencerAnimToolsModule.cpp
@@ -9,9 +9,7 @@
 #include "UnrealEdGlobals.h"
 #include "LevelEditorSequencerIntegration.h"
 
-#include "BaseGizmos/TransformGizmoUtil.h"
 #include "MotionTrailTool.h"
-#include "SequencerAnimEditPivotTool.h"
 #include "Tools/MotionTrailOptions.h"
 
 #define LOCTEXT_NAMESPACE "FSequencerAnimToolsModule"
@@ -23,8 +21,6 @@ namespace SequencerAnimTools
 
 void FSequencerAnimToolsModule::StartupModule()
 {
-	FLevelEditorModule& LevelEditor = FModuleManager::Get().LoadModuleChecked<FLevelEditorModule>("LevelEditor");
-	LevelEditor.OnLevelEditorCreated().AddRaw(this, &FSequencerAnimToolsModule::OnLevelEditorCreated);
 	UMotionTrailToolOptions* MotionTrailOptions = GetMutableDefault<UMotionTrailToolOptions>();
 	if (MotionTrailOptions)
 	{
@@ -34,8 +30,6 @@ void FSequencerAnimToolsModule::StartupModule()
 
 void FSequencerAnimToolsModule::ShutdownModule()
 {
-	FLevelEditorModule& LevelEditor = FModuleManager::Get().LoadModuleChecked<FLevelEditorModule>("LevelEditor");
-	LevelEditor.OnLevelEditorCreated().RemoveAll(this);
 	/* can't do this for some reason the object is already dead (but not nullptr so can't check)
 	UMotionTrailToolOptions* MotionTrailOptions = GetMutableDefault<UMotionTrailToolOptions>();
 	if (MotionTrailOptions)
@@ -45,39 +39,13 @@ void FSequencerAnimToolsModule::ShutdownModule()
 	*/
 }
 
-void FSequencerAnimToolsModule::OnLevelEditorCreated(TSharedPtr<ILevelEditor> InLevelEditor)
-{
-	LevelEditorPtr = InLevelEditor;
-	InLevelEditor->GetEditorModeManager().OnEditorModeIDChanged().AddLambda([this](const FEditorModeID& InModeID, bool IsEnteringMode)
-		{
-			FEditorModeID ModeID = TEXT("EM_SequencerMode");
-			if (ModeID == InModeID && LevelEditorPtr.IsValid())
-			{
-				//don't re-register, will hit an ensure
-				if (AlreadyRegisteredTools.Contains(LevelEditorPtr.Pin().Get()) == false)
-				{
-					UEdMode* EdMode = LevelEditorPtr.Pin()->GetEditorModeManager().GetActiveScriptableMode(ModeID);
-
-					// Register the factories and gizmos if we are entering the new Gizmo Mode
-					if (IsEnteringMode && EdMode)
-					{
-						LevelEditorPtr.Pin()->GetEditorModeManager().GetInteractiveToolsContext()->ToolManager->RegisterToolType(TEXT("SequencerMotionTrail"), NewObject<UMotionTrailToolBuilder>(EdMode));
-						LevelEditorPtr.Pin()->GetEditorModeManager().GetInteractiveToolsContext()->ToolManager->RegisterToolType(TEXT("SequencerPivotTool"), NewObject<USequencerPivotToolBuilder>(EdMode));
-
-						UE::TransformGizmoUtil::RegisterTransformGizmoContextObject(LevelEditorPtr.Pin()->GetEditorModeManager().GetInteractiveToolsContext());
-
-					}
-					AlreadyRegisteredTools.Add(LevelEditorPtr.Pin().Get());
-				}
-			}
-		});
-}
-
 void FSequencerAnimToolsModule::OnMotionTralOptionChanged(FName PropertyName)
 {
+	FLevelEditorModule* LevelEditorModule = FModuleManager::GetModulePtr<FLevelEditorModule>(TEXT("LevelEditor"));
 	UMotionTrailToolOptions* MotionTrailOptions = GetMutableDefault<UMotionTrailToolOptions>();
-	if (MotionTrailOptions)
+	if (MotionTrailOptions && LevelEditorModule)
 	{
+		TWeakPtr<ILevelEditor> LevelEditorPtr = LevelEditorModule->GetLevelEditorInstance();
 		if (PropertyName == FName("bShowTrails") && LevelEditorPtr.IsValid())
 		{
 			static bool bIsChangingTrail = false;
@@ -87,7 +55,6 @@ void FSequencerAnimToolsModule::OnMotionTralOptionChanged(FName PropertyName)
 				bIsChangingTrail = true;
 				if (FLevelEditorSequencerIntegration::Get().GetSequencers().Num() > 0)
 				{
-
 					if (MotionTrailOptions->bShowTrails)
 					{
 						LevelEditorPtr.Pin()->GetEditorModeManager().GetInteractiveToolsContext()->ToolManager->SelectActiveToolType(EToolSide::Left, TEXT("SequencerMotionTrail"));
@@ -112,10 +79,6 @@ void FSequencerAnimToolsModule::OnMotionTralOptionChanged(FName PropertyName)
 	}
 }
 
-
-
-
-
 } // namespace SequencerAnimTools
 } // namespace UE
 
diff --git a/Engine/Plugins/Editor/SequencerAnimTools/Source/SequencerAnimTools/Private/SequencerToolsEditMode.cpp b/Engine/Plugins/Editor/SequencerAnimTools/Source/SequencerAnimTools/Private/SequencerToolsEditMode.cpp
--- a/Engine/Plugins/Editor/SequencerAnimTools/Source/SequencerAnimTools/Private/SequencerToolsEditMode.cpp
+++ b/Engine/Plugins/Editor/SequencerAnimTools/Source/SequencerAnimTools/Private/SequencerToolsEditMode.cpp
@@ -9,6 +9,9 @@
 #include "EditorModeManager.h"
 #include "InteractiveToolManager.h"
 #include "BaseSequencerAnimTool.h"
+#include "BaseGizmos/TransformGizmoUtil.h"
+#include "MotionTrailTool.h"
+#include "SequencerAnimEditPivotTool.h"
 
 #define LOCTEXT_NAMESPACE "SequencerAnimTools"
 
@@ -31,6 +34,42 @@ USequencerToolsEditMode::~USequencerToolsEditMode()
 
 }
 
+void USequencerToolsEditMode::Enter()
+{
+	FLevelEditorModule* LevelEditorModule = FModuleManager::GetModulePtr<FLevelEditorModule>(TEXT("LevelEditor"));
+	
+	if (LevelEditorModule)
+	{
+		TWeakPtr<ILevelEditor> LevelEditorPtr = LevelEditorModule->GetLevelEditorInstance();
+		
+		if (LevelEditorPtr.IsValid())
+		{
+			LevelEditorPtr.Pin()->GetEditorModeManager().GetInteractiveToolsContext()->ToolManager->RegisterToolType(TEXT("SequencerMotionTrail"), NewObject<UMotionTrailToolBuilder>(this));
+			LevelEditorPtr.Pin()->GetEditorModeManager().GetInteractiveToolsContext()->ToolManager->RegisterToolType(TEXT("SequencerPivotTool"), NewObject<USequencerPivotToolBuilder>(this));
+
+			UE::TransformGizmoUtil::RegisterTransformGizmoContextObject(LevelEditorPtr.Pin()->GetEditorModeManager().GetInteractiveToolsContext());
+		}
+	}
+}
+
+void USequencerToolsEditMode::Exit()
+{
+	FLevelEditorModule* LevelEditorModule = FModuleManager::GetModulePtr<FLevelEditorModule>(TEXT("LevelEditor"));
+	
+	if (LevelEditorModule)
+	{
+		TWeakPtr<ILevelEditor> LevelEditorPtr = LevelEditorModule->GetLevelEditorInstance();
+		
+		if (LevelEditorPtr.IsValid())
+		{
+			LevelEditorPtr.Pin()->GetEditorModeManager().GetInteractiveToolsContext()->ToolManager->UnregisterToolType(TEXT("SequencerMotionTrail"));
+			LevelEditorPtr.Pin()->GetEditorModeManager().GetInteractiveToolsContext()->ToolManager->UnregisterToolType(TEXT("SequencerPivotTool"));
+
+			UE::TransformGizmoUtil::DeregisterTransformGizmoContextObject(LevelEditorPtr.Pin()->GetEditorModeManager().GetInteractiveToolsContext());
+		}
+	}
+}
+
 bool USequencerToolsEditMode::IsCompatibleWith(FEditorModeID OtherModeID) const
 {
 	// Compatible with all modes similar to FSequencerEdMode
diff --git a/Engine/Plugins/Editor/SequencerAnimTools/Source/SequencerAnimTools/Public/SequencerAnimToolsModule.h b/Engine/Plugins/Editor/SequencerAnimTools/Source/SequencerAnimTools/Public/SequencerAnimToolsModule.h
--- a/Engine/Plugins/Editor/SequencerAnimTools/Source/SequencerAnimTools/Public/SequencerAnimToolsModule.h
+++ b/Engine/Plugins/Editor/SequencerAnimTools/Source/SequencerAnimTools/Public/SequencerAnimToolsModule.h
@@ -6,8 +6,6 @@
 #include "Framework/Commands/UICommandList.h"
 #include "Modules/ModuleManager.h"
 
-class ILevelEditor;
-
 namespace UE
 {
 namespace SequencerAnimTools
@@ -23,15 +21,10 @@ public:
 	virtual void ShutdownModule() override;
 
 private:
-	TWeakPtr<ILevelEditor> LevelEditorPtr;
 	TSharedPtr<FUICommandList> CommandBindings;
 
 	//delegates
-	void OnLevelEditorCreated(TSharedPtr<ILevelEditor> InLevelEditor);
-	void OnMotionTralOptionChanged(FName PropretyName);
-
-	//need to keep track of which levels we registered the tools for otherwise we can hit an ensure
-	TSet<ILevelEditor*>  AlreadyRegisteredTools;
+	void OnMotionTralOptionChanged(FName PropertyName);
 };
 
 } // namespace SequencerAnimTools
diff --git a/Engine/Plugins/Editor/SequencerAnimTools/Source/SequencerAnimTools/Public/SequencerToolsEditMode.h b/Engine/Plugins/Editor/SequencerAnimTools/Source/SequencerAnimTools/Public/SequencerToolsEditMode.h
--- a/Engine/Plugins/Editor/SequencerAnimTools/Source/SequencerAnimTools/Public/SequencerToolsEditMode.h
+++ b/Engine/Plugins/Editor/SequencerAnimTools/Source/SequencerAnimTools/Public/SequencerToolsEditMode.h
@@ -18,6 +18,10 @@ public:
 	USequencerToolsEditMode();
 	virtual ~USequencerToolsEditMode();
 
+	// UEdMode interface
+	virtual void Enter();
+	virtual void Exit();
+
 	// UBaseLegacyWidgetEdMode interface
 	
 	virtual bool UsesToolkits() const  override { return false; }
