Sequencer: UsesPropertyWidgets for SequencerToolsEditMode so that properties that have bEnable3dWidget use the correct pivot location

diff --git a/Engine/Plugins/Editor/SequencerAnimTools/Source/SequencerAnimTools/Public/SequencerToolsEditMode.h b/Engine/Plugins/Editor/SequencerAnimTools/Source/SequencerAnimTools/Public/SequencerToolsEditMode.h
--- a/Engine/Plugins/Editor/SequencerAnimTools/Source/SequencerAnimTools/Public/SequencerToolsEditMode.h
+++ b/Engine/Plugins/Editor/SequencerAnimTools/Source/SequencerAnimTools/Public/SequencerToolsEditMode.h
@@ -26,6 +26,7 @@ public:
 	
 	virtual bool UsesToolkits() const  override { return false; }
 	virtual bool InputKey(FEditorViewportClient* InViewportClient, FViewport* InViewport, FKey InKey, EInputEvent InEvent) override;
+	virtual bool UsesPropertyWidgets() const override { return true; }
 
 	virtual bool IsCompatibleWith(FEditorModeID OtherModeID) const override;
 
