Sequencer:Control Rig: Don't move parented custom attribute type controls with pivot tool, and make sure to not key if not needed.
#trivial

diff --git a/Engine/Plugins/Animation/ControlRig/Source/ControlRigEditor/Private/ControlRigSequencerEditorLibrary.cpp b/Engine/Plugins/Animation/ControlRig/Source/ControlRigEditor/Private/ControlRigSequencerEditorLibrary.cpp
--- a/Engine/Plugins/Animation/ControlRig/Source/ControlRigEditor/Private/ControlRigSequencerEditorLibrary.cpp
+++ b/Engine/Plugins/Animation/ControlRig/Source/ControlRigEditor/Private/ControlRigSequencerEditorLibrary.cpp
@@ -599,7 +599,7 @@ void UControlRigSequencerEditorLibrary::SetControlRigWorldTransform(ULevelSequen
 	TArray<FTransform> WorldTransforms;
 	WorldTransforms.Add(WorldTransform);
 
-	LocalSetControlRigWorldTransforms(LevelSequence, ControlRig, ControlName, EControlRigSetKey::Always, Frames, WorldTransforms, TimeUnit);
+	LocalSetControlRigWorldTransforms(LevelSequence, ControlRig, ControlName, SetKey, Frames, WorldTransforms, TimeUnit);
 
 }
 
diff --git a/Engine/Plugins/Editor/SequencerAnimTools/Source/SequencerAnimTools/Private/SequencerAnimEditPivotTool.cpp b/Engine/Plugins/Editor/SequencerAnimTools/Source/SequencerAnimTools/Private/SequencerAnimEditPivotTool.cpp
--- a/Engine/Plugins/Editor/SequencerAnimTools/Source/SequencerAnimTools/Private/SequencerAnimEditPivotTool.cpp
+++ b/Engine/Plugins/Editor/SequencerAnimTools/Source/SequencerAnimTools/Private/SequencerAnimEditPivotTool.cpp
@@ -488,18 +488,34 @@ void USequencerPivotTool::GizmoTransformStarted(UTransformProxy* Proxy)
 			if (SelectedControls.Num() > 0)
 			{
 				ControlRig->Modify();
+
 				for (const FName& Name : SelectedControls)
 				{
-					FTransform Transform = UControlRigSequencerEditorLibrary::GetControlRigWorldTransform(LevelSequence, ControlRig.Get(), Name, FrameNumber,
-						ESequenceTimeUnit::TickResolution);
-
-					FControlRigSelectionDuringDrag ControlDrag;
-					ControlDrag.LevelSequence = LevelSequence;
-					ControlDrag.ControlName = Name;
-					ControlDrag.ControlRig = ControlRig.Get();
-					ControlDrag.CurrentFrame = FrameNumber;
-					ControlDrag.CurrentTransform = Transform;
-					ControlRigDrags.Add(ControlDrag);
+					if (FRigControlElement* ControlElement = ControlRig->FindControl(Name))
+					{
+						//if we have a parent that's selected and we are a float/bool/int/enum skip
+						if (FRigControlElement* ParentControlElement = Cast<FRigControlElement>(ControlRig->GetHierarchy()->GetFirstParent(ControlElement)))
+						{
+							if (ControlRig->IsControlSelected(ParentControlElement->GetName()) && (
+								ControlElement->Settings.ControlType == ERigControlType::Bool || 
+								ControlElement->Settings.ControlType == ERigControlType::Float || 
+								ControlElement->Settings.ControlType == ERigControlType::Integer)
+								)
+							{
+								continue;
+							}
+						}
+						FTransform Transform = UControlRigSequencerEditorLibrary::GetControlRigWorldTransform(LevelSequence, ControlRig.Get(), Name, FrameNumber,
+							ESequenceTimeUnit::TickResolution);
+
+						FControlRigSelectionDuringDrag ControlDrag;
+						ControlDrag.LevelSequence = LevelSequence;
+						ControlDrag.ControlName = Name;
+						ControlDrag.ControlRig = ControlRig.Get();
+						ControlDrag.CurrentFrame = FrameNumber;
+						ControlDrag.CurrentTransform = Transform;
+						ControlRigDrags.Add(ControlDrag);
+					}
 				}
 			}
 		}
