Sequencer: Control Rig: Minor fix for keepign state with edit pivot tool
#trivial

diff --git a/Engine/Plugins/Editor/SequencerAnimTools/Source/SequencerAnimTools/Private/SequencerAnimEditPivotTool.cpp b/Engine/Plugins/Editor/SequencerAnimTools/Source/SequencerAnimTools/Private/SequencerAnimEditPivotTool.cpp
--- a/Engine/Plugins/Editor/SequencerAnimTools/Source/SequencerAnimTools/Private/SequencerAnimEditPivotTool.cpp
+++ b/Engine/Plugins/Editor/SequencerAnimTools/Source/SequencerAnimTools/Private/SequencerAnimEditPivotTool.cpp
@@ -331,6 +331,7 @@ bool USequencerPivotTool::SetGizmoBasedOnSelection(bool bUseSaved)
 	GEditor->GetSelectedActors()->GetSelectedObjects(SelectedActors);
 	for (AActor* SelectedActor : SelectedActors)
 	{
+		bHaveSomethingSelected = true;
 		FActorMappings* Mappings = SavedPivotLocations.ActorMappings.Find(SelectedActor);
 		if (bUseSaved && Mappings)
 		{
