Sequencer: Normalize transform before blending

diff --git a/Engine/Plugins/Editor/SequencerAnimTools/Source/SequencerAnimTools/Private/TrajectoryCache.cpp b/Engine/Plugins/Editor/SequencerAnimTools/Source/SequencerAnimTools/Private/TrajectoryCache.cpp
--- a/Engine/Plugins/Editor/SequencerAnimTools/Source/SequencerAnimTools/Private/TrajectoryCache.cpp
+++ b/Engine/Plugins/Editor/SequencerAnimTools/Source/SequencerAnimTools/Private/TrajectoryCache.cpp
@@ -41,8 +41,14 @@ FTransform FArrayTrajectoryCache::GetInterp(const double InTime) const
 		return TrajectoryCache[LowIdx];
 	}
 
+	FTransform KeyAtom1 = TrajectoryCache[LowIdx];
+	FTransform KeyAtom2 = TrajectoryCache[HighIdx];
+	
+	KeyAtom1.NormalizeRotation();
+	KeyAtom2.NormalizeRotation();
+
 	FTransform TempBlended; 
-	TempBlended.Blend(TrajectoryCache[LowIdx], TrajectoryCache[HighIdx], T);
+	TempBlended.Blend(KeyAtom1, KeyAtom2, T);
 	return TempBlended;
 }
 
@@ -70,10 +76,22 @@ void FArrayTrajectoryCache::GetInterp(const double InTime,FTransform& Transform,
 		}
 	}
 
-	Transform.Blend(TrajectoryCache[LowIdx], TrajectoryCache[HighIdx], T);
+	FTransform KeyAtom1 = TrajectoryCache[LowIdx];
+	FTransform KeyAtom2 = TrajectoryCache[HighIdx];
+	
+	KeyAtom1.NormalizeRotation();
+	KeyAtom2.NormalizeRotation();
+
+	Transform.Blend(KeyAtom1, KeyAtom2, T);
 	if (ParentTrajectoryCache.Num() > 0)
 	{
-		ParentTransform.Blend(ParentTrajectoryCache[LowIdx], ParentTrajectoryCache[HighIdx], T);
+		KeyAtom1 = ParentTrajectoryCache[LowIdx];
+		KeyAtom2 = ParentTrajectoryCache[HighIdx];
+
+		KeyAtom1.NormalizeRotation();
+		KeyAtom2.NormalizeRotation();
+		
+		ParentTransform.Blend(KeyAtom1, KeyAtom2, T);
 	}
 	else
 	{
