Sequencer: Motion Trails: Rename Ticks to Marks per RichardB's suggestion for documenation.
#trivial

diff --git a/Engine/Plugins/Editor/SequencerAnimTools/Source/SequencerAnimTools/Private/TrailHierarchy.cpp b/Engine/Plugins/Editor/SequencerAnimTools/Source/SequencerAnimTools/Private/TrailHierarchy.cpp
--- a/Engine/Plugins/Editor/SequencerAnimTools/Source/SequencerAnimTools/Private/TrailHierarchy.cpp
+++ b/Engine/Plugins/Editor/SequencerAnimTools/Source/SequencerAnimTools/Private/TrailHierarchy.cpp
@@ -40,7 +40,7 @@ void FTrailHierarchyRenderer::Render(const FSceneView* View, FViewport* Viewport
 			FDisplayContext DisplayContext = {
 				GuidTrailPair.Key,
 				FTrailScreenSpaceTransform(View, Viewport),
-				UMotionTrailToolOptions::GetTrailOptions()->SecondsPerTick,
+				UMotionTrailToolOptions::GetTrailOptions()->SecondsPerMark,
 				OwningHierarchy->GetViewRange(),
 				OwningHierarchy->GetSecondsPerSegment()
 			};
@@ -91,7 +91,7 @@ void FTrailHierarchyRenderer::Render(const FSceneView* View, FViewport* Viewport
 void FTrailHierarchyRenderer::DrawHUD(FEditorViewportClient* ViewportClient, FViewport* Viewport, const FSceneView* View, FCanvas* Canvas)
 {
 
-	if (UMotionTrailToolOptions::GetTrailOptions()->bShowTicks == false)
+	if (UMotionTrailToolOptions::GetTrailOptions()->bShowMarks == false)
 	{
 		for (const TPair<FGuid, TUniquePtr<FTrail>>& GuidTrailPair : OwningHierarchy->GetAllTrails())
 		{
@@ -102,9 +102,9 @@ void FTrailHierarchyRenderer::DrawHUD(FEditorViewportClient* ViewportClient, FVi
 
 	const FDateTime DrawHUDStartTime = FDateTime::Now();
 
-	const double SecondsPerTick = UMotionTrailToolOptions::GetTrailOptions()->bLockTicksToFrames ? OwningHierarchy->GetSecondsPerFrame() : UMotionTrailToolOptions::GetTrailOptions()->SecondsPerTick;
-	const int32 PredictedNumTicks = int32((OwningHierarchy->GetViewRange().Size<double>() / SecondsPerTick) * OwningHierarchy->GetAllTrails().Num() * 1.3); // Multiply by 1.3 to be safe
-	Canvas->GetBatchedElements(FCanvas::EElementType::ET_Line)->AddReserveLines(PredictedNumTicks);
+	const double SecondsPerMark = UMotionTrailToolOptions::GetTrailOptions()->bLockMarksToFrames ? OwningHierarchy->GetSecondsPerFrame() : UMotionTrailToolOptions::GetTrailOptions()->SecondsPerMark;
+	const int32 PredictedNumMarks = int32((OwningHierarchy->GetViewRange().Size<double>() / SecondsPerMark) * OwningHierarchy->GetAllTrails().Num() * 1.3); // Multiply by 1.3 to be safe
+	Canvas->GetBatchedElements(FCanvas::EElementType::ET_Line)->AddReserveLines(PredictedNumMarks);
 
 	for(const TPair<FGuid, TUniquePtr<FTrail>>& GuidTrailPair : OwningHierarchy->GetAllTrails())
 	{
@@ -114,19 +114,19 @@ void FTrailHierarchyRenderer::DrawHUD(FEditorViewportClient* ViewportClient, FVi
 			FDisplayContext DisplayContext = {
 				GuidTrailPair.Key,
 				FTrailScreenSpaceTransform(View, Viewport, ViewportClient->GetDPIScale()),
-				SecondsPerTick,
+				SecondsPerMark,
 				OwningHierarchy->GetViewRange(),
 				OwningHierarchy->GetSecondsPerSegment()
 
 			};
 
-			TArray<FVector2D> Ticks, TickNormals;
-			GuidTrailPair.Value->GetTickPointsForDisplay(DisplayContext, Ticks, TickNormals);
-			const FLinearColor Color = OwningHierarchy->GetVisibilityManager().IsTrailAlwaysVisible(GuidTrailPair.Key) ? CurDrawInfo->GetColor() : UMotionTrailToolOptions::GetTrailOptions()->TickColor;
-			for (int32 Idx = 0; Idx < Ticks.Num(); Idx++)
+			TArray<FVector2D> Marks, MarkNormals;
+			GuidTrailPair.Value->GetTickPointsForDisplay(DisplayContext, Marks, MarkNormals);
+			const FLinearColor Color = OwningHierarchy->GetVisibilityManager().IsTrailAlwaysVisible(GuidTrailPair.Key) ? CurDrawInfo->GetColor() : UMotionTrailToolOptions::GetTrailOptions()->MarkColor;
+			for (int32 Idx = 0; Idx < Marks.Num(); Idx++)
 			{
-				const FVector2D StartPoint = Ticks[Idx] - TickNormals[Idx] * UMotionTrailToolOptions::GetTrailOptions()->TickSize;
-				const FVector2D EndPoint = Ticks[Idx] + TickNormals[Idx] * UMotionTrailToolOptions::GetTrailOptions()->TickSize;
+				const FVector2D StartPoint = Marks[Idx] - MarkNormals[Idx] * UMotionTrailToolOptions::GetTrailOptions()->MarkSize;
+				const FVector2D EndPoint = Marks[Idx] + MarkNormals[Idx] * UMotionTrailToolOptions::GetTrailOptions()->MarkSize;
 				FCanvasLineItem LineItem = FCanvasLineItem(StartPoint, EndPoint);
 				LineItem.SetColor(Color);
 				Canvas->DrawItem(LineItem);
diff --git a/Engine/Source/Editor/Sequencer/Public/Tools/MotionTrailOptions.h b/Engine/Source/Editor/Sequencer/Public/Tools/MotionTrailOptions.h
--- a/Engine/Source/Editor/Sequencer/Public/Tools/MotionTrailOptions.h
+++ b/Engine/Source/Editor/Sequencer/Public/Tools/MotionTrailOptions.h
@@ -28,11 +28,11 @@ public:
 		, bShowFrameNumber(true)
 		, KeyColor(1.0, 1.0, 1.0)
 		, KeySize(10.0f)
-		, bShowTicks(false)
-		, TickColor(0.25,1.0,0.15)
-		, TickSize(5.0)
-		, bLockTicksToFrames(true)
-		, SecondsPerTick(0.1)
+		, bShowMarks(false)
+		, MarkColor(0.25,1.0,0.15)
+		, MarkSize(5.0)
+		, bLockMarksToFrames(true)
+		, SecondsPerMark(0.1)
 	{}
 
 	UPROPERTY(EditAnywhere, Category = Trail);
@@ -68,20 +68,20 @@ public:
 	UPROPERTY(EditAnywhere, Category = Keys, Meta = (ClampMin = "0.0"))
 	double KeySize;
 
-	UPROPERTY(EditAnywhere, Category = Ticks)
-	bool bShowTicks;
+	UPROPERTY(EditAnywhere, Category = Marks)
+	bool bShowMarks;
 
-	UPROPERTY(EditAnywhere, Category = Ticks)
-	FLinearColor TickColor;
+	UPROPERTY(EditAnywhere, Category = Marks)
+	FLinearColor MarkColor;
 
-	UPROPERTY(EditAnywhere, Category = Ticks, Meta = (ClampMin = "0.0"))
-	double TickSize;
+	UPROPERTY(EditAnywhere, Category = Marks, Meta = (ClampMin = "0.0"))
+	double MarkSize;
 
-	UPROPERTY(EditAnywhere, Category = Ticks)
-	bool bLockTicksToFrames;
+	UPROPERTY(EditAnywhere, Category = Marks)
+	bool bLockMarksToFrames;
 
-	UPROPERTY(EditAnywhere, Category = Ticks, Meta = (EditCondition = "!bLockTicksToFrames", ClampMin = "0.01"))
-	double SecondsPerTick;
+	UPROPERTY(EditAnywhere, Category = Marks, Meta = (EditCondition = "!bLockMarksToFrames", ClampMin = "0.01"))
+	double SecondsPerMark;
 
 	virtual void PostEditChangeProperty(struct FPropertyChangedEvent& PropertyChangedEvent) override
 	{
