Sequencer: Anim Pivot: Rotate about pivot, shift will turn it off

diff --git a/Engine/Plugins/Editor/SequencerAnimTools/Source/SequencerAnimTools/Private/SequencerAnimEditPivotTool.cpp b/Engine/Plugins/Editor/SequencerAnimTools/Source/SequencerAnimTools/Private/SequencerAnimEditPivotTool.cpp
--- a/Engine/Plugins/Editor/SequencerAnimTools/Source/SequencerAnimTools/Private/SequencerAnimEditPivotTool.cpp
+++ b/Engine/Plugins/Editor/SequencerAnimTools/Source/SequencerAnimTools/Private/SequencerAnimEditPivotTool.cpp
@@ -351,7 +351,6 @@ bool USequencerPivotTool::SetGizmoBasedOnSelection(bool bUseSaved)
 		GizmoTransform.SetLocation(AverageLocation);
 	}
 	GizmoTransform.SetScale3D(FVector(1.0f, 1.0f, 1.0f));
-	GizmoTransform.SetRotation(FQuat::Identity);
 	return bHaveSomethingSelected;
 }
 void USequencerPivotTool::Shutdown(EToolShutdownType ShutdownType)
@@ -559,7 +558,7 @@ void USequencerPivotTool::GizmoTransformStarted(UTransformProxy* Proxy)
 
 //For some reason when just setting the location of the actor the gizmo was not updating, so we do the full suite of notifications to make
 //sure the other systems know the actor got moved
-static void SetLocation(AActor* Actor, const FVector& NewLocation)
+static void SetLocation(AActor* Actor, const FVector& NewLocation,const TOptional<FQuat>& NewQuat)
 {
 	if (Actor != nullptr)
 	{
@@ -571,6 +570,10 @@ static void SetLocation(AActor* Actor, const FVector& NewLocation)
 			PropertyChain.AddHead(TransformProperty);
 			FCoreUObjectDelegates::OnPreObjectPropertyChanged.Broadcast(Actor, PropertyChain);
 			RootComponent->SetWorldLocation(NewLocation);
+			if (NewQuat.IsSet())
+			{
+				RootComponent->SetWorldRotation(NewQuat.GetValue());
+			}
 			FPropertyChangedEvent PropertyChangedEvent(TransformProperty, EPropertyChangeType::ValueSet);
 			// Broadcast Post Edit change notification, we can't call PostEditChangeProperty directly on Actor or ActorComponent from here since it wasn't pair with a proper PreEditChange
 			FCoreUObjectDelegates::OnObjectPropertyChanged.Broadcast(Actor, PropertyChangedEvent);
@@ -591,6 +594,8 @@ void USequencerPivotTool::GizmoTransformChanged(UTransformProxy* Proxy, FTransfo
 	}
 	bManipulatorMadeChange = true;
 
+	FModifierKeysState KeyState = FSlateApplication::Get().GetModifierKeys();
+	bool bAlsoCalcRotation = KeyState.IsShiftDown() == false;
 	GizmoTransform = Transform;
 	FTransform Diff = Transform.GetRelativeTransform(StartDragTransform);
 	if (Diff.GetRotation().IsIdentity(1e-4f) == false)
@@ -609,6 +614,12 @@ void USequencerPivotTool::GizmoTransformChanged(UTransformProxy* Proxy, FTransfo
 					FVector RotatedDiff = Diff.GetRotation().RotateVector(LocDiff);
 					FVector NewLocation = StartDragTransform.TransformPosition(RotatedDiff);
 					ControlDrag.CurrentTransform.SetLocation(NewLocation);
+					TOptional <FQuat> OptQuat;
+					if (bAlsoCalcRotation)
+					{
+						OptQuat = ControlDrag.CurrentTransform.GetRotation() * Diff.GetRotation();
+						ControlDrag.CurrentTransform.SetRotation(OptQuat.GetValue());
+					}
 					UControlRigSequencerEditorLibrary::SetControlRigWorldTransform(ControlDrag.LevelSequence, ControlDrag.ControlRig, ControlDrag.ControlName,
 						ControlDrag.CurrentFrame, ControlDrag.CurrentTransform, ESequenceTimeUnit::TickResolution, bSetKey);
 				}
@@ -634,7 +645,13 @@ void USequencerPivotTool::GizmoTransformChanged(UTransformProxy* Proxy, FTransfo
 					FVector RotatedDiff = Diff.GetRotation().RotateVector(LocDiff);
 					FVector NewLocation = StartDragTransform.TransformPosition(RotatedDiff);
 					ActorDrag.CurrentTransform.SetLocation(NewLocation);
-					SetLocation(ActorDrag.Actor, NewLocation);
+					TOptional <FQuat> OptQuat;
+					if (bAlsoCalcRotation)
+					{
+						OptQuat = ActorDrag.CurrentTransform.GetRotation()* Diff.GetRotation();
+						ActorDrag.CurrentTransform.SetRotation(OptQuat.GetValue());
+					}
+					SetLocation(ActorDrag.Actor, NewLocation, OptQuat);
 					/*Note that calling the bit below isn't enought. That's just sends a begin/end event, the other bits are needed, and needed between these events
 					UEditorActorSubsystem* EditorActorSubsystem = GEditor->GetEditorSubsystem<UEditorActorSubsystem>();
 					EditorActorSubsystem->SetActorTransform(ActorDrag.Actor,ActorDrag.CurrentTransform);
@@ -644,14 +661,14 @@ void USequencerPivotTool::GizmoTransformChanged(UTransformProxy* Proxy, FTransfo
 			}
 			else //last one with shift we keep locked!
 			{
-				SetLocation(ActorDrag.Actor, ActorDrag.CurrentTransform.GetLocation());
+				TOptional <FQuat> OptQuat;
+				SetLocation(ActorDrag.Actor, ActorDrag.CurrentTransform.GetLocation(),OptQuat);
 			}
 			++Index;
 		}
 	}
 	
 	StartDragTransform = Transform;
-	UpdateGizmoTransform();
 
 }
 
