Sequencer: Initialize motion trail transforms, otherwise it leads to nan when blending

diff --git a/Engine/Plugins/Editor/SequencerAnimTools/Source/SequencerAnimTools/Public/TrajectoryCache.h b/Engine/Plugins/Editor/SequencerAnimTools/Source/SequencerAnimTools/Public/TrajectoryCache.h
--- a/Engine/Plugins/Editor/SequencerAnimTools/Source/SequencerAnimTools/Public/TrajectoryCache.h
+++ b/Engine/Plugins/Editor/SequencerAnimTools/Source/SequencerAnimTools/Public/TrajectoryCache.h
@@ -95,6 +95,11 @@ public:
 
 	{
 		TrajectoryCache.SetNumUninitialized(int32(InTrackRange.Size<double>() / InSpacing) + 1);
+
+		for (FTransform& Transform : TrajectoryCache)
+		{
+			Transform = InDefault;
+		}
 	}
 
 	virtual const FTransform& Get(const double InTime) const override 
