Sequencer: Animation Tool: Pivot: When the Pivot Tool is activiated you will now:
Have the option to enter either the Edit/Free state, where you can move the pivot where you want
Or go into the Pivot/Pose state where you can rotate or translate(which is new) against the pivot.  there is an overlay to let you switch or you can press P to toggle also.
Based on animators feedback.

diff --git a/Engine/Plugins/Editor/SequencerAnimTools/Source/SequencerAnimTools/Private/SequencerAnimEditPivotTool.cpp b/Engine/Plugins/Editor/SequencerAnimTools/Source/SequencerAnimTools/Private/SequencerAnimEditPivotTool.cpp
--- a/Engine/Plugins/Editor/SequencerAnimTools/Source/SequencerAnimTools/Private/SequencerAnimEditPivotTool.cpp
+++ b/Engine/Plugins/Editor/SequencerAnimTools/Source/SequencerAnimTools/Private/SequencerAnimEditPivotTool.cpp
@@ -37,12 +37,25 @@
 #include "UnrealEdGlobals.h"
 #include "Editor/UnrealEdEngine.h"
 #include "Subsystems/EditorActorSubsystem.h"
+#include "LevelEditorViewport.h"
+#include "Modules/ModuleManager.h"
+#include "IAssetViewport.h"
+#include "Widgets/SWindow.h"
+#include "Widgets/Layout/SBorder.h"
+#include "Widgets/Text/STextBlock.h"
+#include "Widgets/Layout/SBox.h"
+#include "Widgets/Input/SButton.h"
+#include "EditorStyleSet.h"
+#include "Viewports/InViewportUIDragOperation.h"
+#include "Widgets/Layout/SUniformGridPanel.h"
+#include "Styling/StyleColors.h"
 
 #define LOCTEXT_NAMESPACE "SequencerAnimTools"
 
 void FEditPivotCommands::RegisterCommands()
 {
 	UI_COMMAND(ResetPivot, "Reset Pivot To Original", "Reset pivot back to original location", EUserInterfaceActionType::Button, FInputChord(EModifierKey::Control|EModifierKey::Shift|EModifierKey::Alt, EKeys::G));
+	UI_COMMAND(ToggleFreePivot, "Toggle Edit/Pose", "Toggle gizmo to move freely or to pivot the selection", EUserInterfaceActionType::Button, FInputChord(EKeys::P));
 }
 
 
@@ -194,10 +207,13 @@ void USequencerPivotTool::Setup()
 		Commands.ResetPivot,
 		FExecuteAction::CreateUObject(this, &USequencerPivotTool::ResetPivot)
 		);
+	CommandBindings->MapAction(
+		Commands.ToggleFreePivot,
+		FExecuteAction::CreateUObject(this, &USequencerPivotTool::TogglePivotMode)
+	);
 
+	CreateAndShowPivotOverlay();
 }
-
-
 void USequencerPivotTool::ResetPivot()
 {
 	SetGizmoBasedOnSelection(false);
@@ -239,7 +255,7 @@ void USequencerPivotTool::SaveLastSelected()
 	}
 }
 
-//when we enter the tool if we have things selected we get it's last transform position
+//when we enter the tool if we have things selected we get its last transform position
 //if not we should have the last thing selected, and if so we select that.
 void USequencerPivotTool::UpdateTransformAndSelectionOnEntering()
 {
@@ -357,6 +373,7 @@ void USequencerPivotTool::Shutdown(EToolShutdownType ShutdownType)
 {
 	GizmoManager->DestroyAllGizmosByOwner(this);
 	RemoveDelegates();
+	RemoveAndDestroyPivotOverlay();
 }
 
 void USequencerPivotTool::DeactivateMe()
@@ -598,28 +615,42 @@ void USequencerPivotTool::GizmoTransformChanged(UTransformProxy* Proxy, FTransfo
 	bool bAlsoCalcRotation = KeyState.IsShiftDown() == false;
 	GizmoTransform = Transform;
 	FTransform Diff = Transform.GetRelativeTransform(StartDragTransform);
-	if (Diff.GetRotation().IsIdentity(1e-4f) == false)
+	const bool bRotationChanged = (Diff.GetRotation().IsIdentity(1e-4f) == false);
+	const bool bTranslationChanged = (Diff.GetTranslation().IsNearlyZero() == false);
+	if (bRotationChanged || bTranslationChanged)
 	{
 		const bool bSetKey = false;
 		int32 Index = 0;
 		for (FControlRigSelectionDuringDrag& ControlDrag : ControlRigDrags)
 		{
-			
-			if (!bShiftPressedWhenStarted || Index != (ControlRigDrags.Num() - 1))
+			if (bInPivotMode == false)
+			{
+				ControlDrag.CurrentTransform = Transform;
+			}
+			else if (!bShiftPressedWhenStarted || Index != (ControlRigDrags.Num() - 1))
 			{
-				FVector LocDiff = ControlDrag.CurrentTransform.GetLocation() - Transform.GetLocation();
-				if (LocDiff.IsNearlyZero(1e-4f) == false)
+				if (bRotationChanged)
 				{
-					LocDiff = StartDragTransform.InverseTransformPosition(ControlDrag.CurrentTransform.GetLocation());
-					FVector RotatedDiff = Diff.GetRotation().RotateVector(LocDiff);
-					FVector NewLocation = StartDragTransform.TransformPosition(RotatedDiff);
-					ControlDrag.CurrentTransform.SetLocation(NewLocation);
-					TOptional <FQuat> OptQuat;
-					if (bAlsoCalcRotation)
+					FVector LocDiff = ControlDrag.CurrentTransform.GetLocation() - Transform.GetLocation();
+					if (LocDiff.IsNearlyZero(1e-4f) == false)
 					{
-						OptQuat = ControlDrag.CurrentTransform.GetRotation() * Diff.GetRotation();
-						ControlDrag.CurrentTransform.SetRotation(OptQuat.GetValue());
+						LocDiff = StartDragTransform.InverseTransformPosition(ControlDrag.CurrentTransform.GetLocation());
+						FVector RotatedDiff = Diff.GetRotation().RotateVector(LocDiff);
+						FVector NewLocation = StartDragTransform.TransformPosition(RotatedDiff);
+						ControlDrag.CurrentTransform.SetLocation(NewLocation);
+						if (bAlsoCalcRotation)
+						{
+							FQuat Quat = ControlDrag.CurrentTransform.GetRotation() * Diff.GetRotation();
+							ControlDrag.CurrentTransform.SetRotation(Quat);
+						}
+						UControlRigSequencerEditorLibrary::SetControlRigWorldTransform(ControlDrag.LevelSequence, ControlDrag.ControlRig, ControlDrag.ControlName,
+							ControlDrag.CurrentFrame, ControlDrag.CurrentTransform, ESequenceTimeUnit::TickResolution, bSetKey);
 					}
+				}
+				else if (bTranslationChanged)
+				{
+					FVector DiffTranslation = StartDragTransform.GetRotation().RotateVector(Diff.GetTranslation());
+					ControlDrag.CurrentTransform.AddToTranslation(DiffTranslation);
 					UControlRigSequencerEditorLibrary::SetControlRigWorldTransform(ControlDrag.LevelSequence, ControlDrag.ControlRig, ControlDrag.ControlName,
 						ControlDrag.CurrentFrame, ControlDrag.CurrentTransform, ESequenceTimeUnit::TickResolution, bSetKey);
 				}
@@ -636,27 +667,42 @@ void USequencerPivotTool::GizmoTransformChanged(UTransformProxy* Proxy, FTransfo
 		Index = 0;
 		for (FActorSelectonDuringDrag& ActorDrag : ActorDrags)
 		{
-			if (!bShiftPressedWhenStarted || Index != (ActorDrags.Num() - 1))
+			if (bInPivotMode == false)
 			{
-				FVector LocDiff = ActorDrag.CurrentTransform.GetLocation() - Transform.GetLocation();
-				if (LocDiff.IsNearlyZero(1e-4f) == false)
+				ActorDrag.CurrentTransform = Transform;
+			}
+			else if (!bShiftPressedWhenStarted || Index != (ActorDrags.Num() - 1))
+			{
+				if (bRotationChanged)
 				{
-					LocDiff = StartDragTransform.InverseTransformPosition(ActorDrag.CurrentTransform.GetLocation());
-					FVector RotatedDiff = Diff.GetRotation().RotateVector(LocDiff);
-					FVector NewLocation = StartDragTransform.TransformPosition(RotatedDiff);
-					ActorDrag.CurrentTransform.SetLocation(NewLocation);
-					TOptional <FQuat> OptQuat;
-					if (bAlsoCalcRotation)
+					FVector LocDiff = ActorDrag.CurrentTransform.GetLocation() - Transform.GetLocation();
+					if (LocDiff.IsNearlyZero(1e-4f) == false)
 					{
-						OptQuat = ActorDrag.CurrentTransform.GetRotation()* Diff.GetRotation();
-						ActorDrag.CurrentTransform.SetRotation(OptQuat.GetValue());
+						LocDiff = StartDragTransform.InverseTransformPosition(ActorDrag.CurrentTransform.GetLocation());
+						FVector RotatedDiff = Diff.GetRotation().RotateVector(LocDiff);
+						FVector NewLocation = StartDragTransform.TransformPosition(RotatedDiff);
+						ActorDrag.CurrentTransform.SetLocation(NewLocation);
+						TOptional <FQuat> OptQuat;
+						if (bAlsoCalcRotation)
+						{
+							OptQuat = ActorDrag.CurrentTransform.GetRotation() * Diff.GetRotation();
+							ActorDrag.CurrentTransform.SetRotation(OptQuat.GetValue());
+						}
+						SetLocation(ActorDrag.Actor, NewLocation, OptQuat);
+						/*Note that calling the bit below isn't enought. That's just sends a begin/end event, the other bits are needed, and needed between these events
+						UEditorActorSubsystem* EditorActorSubsystem = GEditor->GetEditorSubsystem<UEditorActorSubsystem>();
+						EditorActorSubsystem->SetActorTransform(ActorDrag.Actor,ActorDrag.CurrentTransform);
+						GUnrealEd->UpdatePivotLocationForSelection();
+						*/
 					}
+				}
+				else if (bTranslationChanged)
+				{
+					FVector DiffTranslation = StartDragTransform.GetRotation().RotateVector(Diff.GetTranslation());
+					ActorDrag.CurrentTransform.AddToTranslation(DiffTranslation);
+					const FVector NewLocation = ActorDrag.CurrentTransform.GetLocation();
+					TOptional <FQuat> OptQuat;
 					SetLocation(ActorDrag.Actor, NewLocation, OptQuat);
-					/*Note that calling the bit below isn't enought. That's just sends a begin/end event, the other bits are needed, and needed between these events
-					UEditorActorSubsystem* EditorActorSubsystem = GEditor->GetEditorSubsystem<UEditorActorSubsystem>();
-					EditorActorSubsystem->SetActorTransform(ActorDrag.Actor,ActorDrag.CurrentTransform);
-					GUnrealEd->UpdatePivotLocationForSelection();
-					*/
 				}
 			}
 			else //last one with shift we keep locked!
@@ -785,7 +831,6 @@ void USequencerPivotTool::OnPropertyModified(UObject* PropertySet, FProperty* Pr
 
 void USequencerPivotTool::Render(IToolsContextRenderAPI* RenderAPI)
 {
-
 	if (bPickingPivotLocation == false )
 	{
 		FPrimitiveDrawInterface* PDI = RenderAPI->GetPrimitiveDrawInterface();
@@ -807,7 +852,265 @@ void USequencerPivotTool::Render(IToolsContextRenderAPI* RenderAPI)
 	}
 }
 
+//Pivot Overlay functions
+
+FVector2D USequencerPivotTool::LastPivotOverlayLocation = FVector2D(0.0, 0.0);
+
+class SPivotOverlayWidget : public SCompoundWidget
+{
+	SLATE_BEGIN_ARGS(SPivotOverlayWidget) {}
+	SLATE_ARGUMENT(USequencerPivotTool*, InPivotTool)
+	SLATE_END_ARGS()
+	~SPivotOverlayWidget()
+	{
+	}
+
+	void Construct(const FArguments& InArgs)
+	{
+		OwningPivotTool = InArgs._InPivotTool;
+		ChildSlot
+		[
+			SNew(SBorder)
+			.BorderImage(FAppStyle::Get().GetBrush("EditorViewport.OverlayBrush"))
+			.Padding(10.f)
+			[
+				SNew(SVerticalBox)
+				+ SVerticalBox::Slot()
+				.AutoHeight()
+				.HAlign(HAlign_Center)
+				.Padding(0.0f, 0.0f, 0.0f, 2.0f)
+				[
+					SNew(STextBlock)
+					.TextStyle(FAppStyle::Get(), "NormalText.Important")
+					.Text(LOCTEXT("PivotTool", "Pivot Tool"))
+				]
+				+ SVerticalBox::Slot()
+				.AutoHeight()
+				.HAlign(HAlign_Center)
+				[
+					SNew(SUniformGridPanel)
+					.SlotPadding(2)
+					+ SUniformGridPanel::Slot(0, 0)
+					[	
+						SNew(SButton)
+						.OnClicked_Raw(this, &SPivotOverlayWidget::OnButtonClicked)
+						.Text(LOCTEXT("Edit", "Edit"))
+						.ToolTipText(LOCTEXT("Edit_Tooltip","When in edit mode you can move the pivot location freely without changing the object's transform"))
+						.ButtonColorAndOpacity(this, &SPivotOverlayWidget::GetEditColor)				
+					]
+					+ SUniformGridPanel::Slot(1, 0)
+					[							
+						SNew(SButton)
+						.OnClicked_Raw(this, &SPivotOverlayWidget::OnButtonClicked)
+						.Text(LOCTEXT("Pose", "Pose"))
+						.ToolTipText(LOCTEXT("Pose_Tooltip", "In pose mode you will rotate or translate the object about its pivot"))
+						.ButtonColorAndOpacity(this, &SPivotOverlayWidget::GetPoseColor)
+					]
+				]
+			]
+		];
+	}
+
+private:
+	TWeakObjectPtr<USequencerPivotTool> OwningPivotTool;
 
+	FReply OnButtonClicked()
+	{
+		if (OwningPivotTool.IsValid())
+		{
+			OwningPivotTool.Get()->TogglePivotMode();
+			return FReply::Handled();
+		}
+		return FReply::Unhandled();
+	}
+
+	FSlateColor GetEditColor() const
+	{
+		if (OwningPivotTool.IsValid())
+		{
+			return OwningPivotTool.Get()->IsInPivotMode() ? FStyleColors::Transparent : FStyleColors::Select;
+		}
+		return FStyleColors::Transparent;
+	}
+
+	FSlateColor GetPoseColor() const
+	{
+		if (OwningPivotTool.IsValid())
+		{
+			return OwningPivotTool.Get()->IsInPivotMode() ? FStyleColors::Select : FStyleColors::Transparent;
+		}
+		return FStyleColors::Transparent;
+	}
+
+	FReply OnMouseButtonDown(const FGeometry& MyGeometry, const FPointerEvent& MouseEvent) override
+	{
+		return FReply::Handled().DetectDrag(SharedThis(this), EKeys::LeftMouseButton);
+	}
+
+	FReply OnDragDetected(const FGeometry& MyGeometry, const FPointerEvent& MouseEvent)
+	{
+		// Need to remember where within a tab we grabbed
+		const FVector2D TabGrabScreenSpaceOffset = MouseEvent.GetScreenSpacePosition() - MyGeometry.GetAbsolutePosition();
+
+		FOnInViewportUIDropped OnUIDropped = FOnInViewportUIDropped::CreateRaw(this, &SPivotOverlayWidget::FinishDraggingWidget);
+		// Start dragging.
+		TSharedRef<FInViewportUIDragOperation> DragDropOperation =
+			FInViewportUIDragOperation::New(
+				SharedThis(this),
+				TabGrabScreenSpaceOffset,
+				GetDesiredSize(),
+				OnUIDropped
+			);
+
+		if (OwningPivotTool.IsValid())
+		{
+			OwningPivotTool.Get()->TryRemovePivotOverlay();
+		}
+		return FReply::Handled().BeginDragDrop(DragDropOperation);
+	}
+
+	void FinishDraggingWidget(const FVector2D InLocation)
+	{
+		if (OwningPivotTool.IsValid())
+		{
+			if (FLevelEditorModule* LevelEditorModule = FModuleManager::GetModulePtr<FLevelEditorModule>(TEXT("LevelEditor")))
+			{
+				TSharedPtr<IAssetViewport> ActiveLevelViewport = LevelEditorModule->GetFirstActiveViewport();
+				if (ActiveLevelViewport.IsValid())
+				{
+					OwningPivotTool.Get()->UpdatePivotOverlayLocation(InLocation, ActiveLevelViewport);
+					OwningPivotTool.Get()->TryShowPivotOverlay();
+				}
+			}
+		}
+	}
+};
+
+static FVector2D GetActiveViewportSize(TSharedPtr<IAssetViewport>& ActiveViewport)
+{
+	FSceneViewFamilyContext ViewFamily(FSceneViewFamily::ConstructionValues(
+		ActiveViewport->GetActiveViewport(),
+		ActiveViewport->GetAssetViewportClient().GetScene(),
+		ActiveViewport->GetAssetViewportClient().EngineShowFlags)
+		.SetRealtimeUpdate(ActiveViewport->GetAssetViewportClient().IsRealtime()));
+	// SceneView is deleted with the ViewFamily
+	FSceneView* SceneView = ActiveViewport->GetAssetViewportClient().CalcSceneView(&ViewFamily);
+	const float InvDpiScale = 1.0f / ActiveViewport->GetAssetViewportClient().GetDPIScale();
+	const float MaxX = SceneView->UnscaledViewRect.Width() * InvDpiScale;
+	const float MaxY = SceneView->UnscaledViewRect.Height() * InvDpiScale;
+	return FVector2D(MaxX, MaxY);
+}
+
+void USequencerPivotTool::CreateAndShowPivotOverlay()
+{
+	if (FLevelEditorModule* LevelEditorModule = FModuleManager::GetModulePtr<FLevelEditorModule>(TEXT("LevelEditor")))
+	{
+		TSharedPtr<IAssetViewport> ActiveLevelViewport = LevelEditorModule->GetFirstActiveViewport();
+		if (ActiveLevelViewport.IsValid())
+		{
+			FVector2D NewWidgetLocation = LastPivotOverlayLocation;
+
+			if (NewWidgetLocation.IsZero())
+			{
+				const FVector2D ActiveViewportSize = GetActiveViewportSize(ActiveLevelViewport);
+				NewWidgetLocation.X = ActiveViewportSize.X / 2.0f;
+				NewWidgetLocation.Y = 50.0f;			
+
+			}
+			UpdatePivotOverlayLocation(NewWidgetLocation, ActiveLevelViewport);
+
+			TAttribute<FMargin> MarginPadding = TAttribute<FMargin>::Create(TAttribute<FMargin>::FGetter::CreateUObject(this, &USequencerPivotTool::GetPivotOverlayPadding));
+
+			SAssignNew(PivotWidget, SHorizontalBox)
+
+				+ SHorizontalBox::Slot()
+				.FillWidth(1.0f)
+				.VAlign(VAlign_Top)
+				.HAlign(HAlign_Left)
+				.Padding(MarginPadding)
+				[
+					SNew(SPivotOverlayWidget)
+					.InPivotTool(this)
+				];
+
+			if (PivotWidget)
+			{
+				ActiveLevelViewport->AddOverlayWidget(PivotWidget.ToSharedRef());
+			}
+		}
+	}
+}
+
+FMargin USequencerPivotTool::GetPivotOverlayPadding() const
+{
+	return FMargin(LastPivotOverlayLocation.X, LastPivotOverlayLocation.Y, 0, 0);
+}
+
+void USequencerPivotTool::RemoveAndDestroyPivotOverlay()
+{
+	if (FLevelEditorModule* LevelEditorModule = FModuleManager::GetModulePtr<FLevelEditorModule>(TEXT("LevelEditor")))
+	{
+		TSharedPtr<IAssetViewport> ActiveLevelViewport = LevelEditorModule->GetFirstActiveViewport();
+		if (ActiveLevelViewport.IsValid())
+		{
+			if (PivotWidget)
+			{
+				ActiveLevelViewport->RemoveOverlayWidget(PivotWidget.ToSharedRef());
+				PivotWidget.Reset();
+			}
+		}
+	}
+}
+
+void USequencerPivotTool::TryRemovePivotOverlay()
+{
+	if (FLevelEditorModule* LevelEditorModule = FModuleManager::GetModulePtr<FLevelEditorModule>(TEXT("LevelEditor")))
+	{
+		TSharedPtr<IAssetViewport> ActiveLevelViewport = LevelEditorModule->GetFirstActiveViewport();
+		if (ActiveLevelViewport.IsValid())
+		{
+			if (PivotWidget)
+			{
+				ActiveLevelViewport->RemoveOverlayWidget(PivotWidget.ToSharedRef());
+			}
+		}
+	}
+}
+
+void USequencerPivotTool::TryShowPivotOverlay()
+{
+	if (FLevelEditorModule* LevelEditorModule = FModuleManager::GetModulePtr<FLevelEditorModule>(TEXT("LevelEditor")))
+	{
+		TSharedPtr<IAssetViewport> ActiveLevelViewport = LevelEditorModule->GetFirstActiveViewport();
+		if (ActiveLevelViewport.IsValid())
+		{
+			if (PivotWidget)
+			{
+				ActiveLevelViewport->AddOverlayWidget(PivotWidget.ToSharedRef());
+			}
+		}
+	}
+}
+
+void USequencerPivotTool::UpdatePivotOverlayLocation(const FVector2D InLocation, TSharedPtr<IAssetViewport> ActiveLevelViewport)
+{
+	const FVector2D ActiveViewportSize = GetActiveViewportSize(ActiveLevelViewport);
+	FVector2D ScreenPos = InLocation;
+
+	const float EdgeFactor = 0.97f;
+	const float MinX = ActiveViewportSize.X * (1 - EdgeFactor);
+	const float MinY = ActiveViewportSize.Y * (1 - EdgeFactor);
+	const float MaxX = ActiveViewportSize.X * EdgeFactor;
+	const float MaxY = ActiveViewportSize.Y * EdgeFactor;
+	const bool bOutside = ScreenPos.X < MinX || ScreenPos.X > MaxX || ScreenPos.Y < MinY || ScreenPos.Y > MaxY;
+	if (bOutside)
+	{
+		// reset the location if it was placed out of bounds
+		ScreenPos.X = ActiveViewportSize.X / 2.0f;
+		ScreenPos.Y = 50.0f;
+	}
+	LastPivotOverlayLocation = ScreenPos;
+}
 
 
 #undef LOCTEXT_NAMESPACE
diff --git a/Engine/Plugins/Editor/SequencerAnimTools/Source/SequencerAnimTools/Public/SequencerAnimEditPivotTool.h b/Engine/Plugins/Editor/SequencerAnimTools/Source/SequencerAnimTools/Public/SequencerAnimEditPivotTool.h
--- a/Engine/Plugins/Editor/SequencerAnimTools/Source/SequencerAnimTools/Public/SequencerAnimEditPivotTool.h
+++ b/Engine/Plugins/Editor/SequencerAnimTools/Source/SequencerAnimTools/Public/SequencerAnimEditPivotTool.h
@@ -109,6 +109,8 @@ public:
 public:
 	/** Reset the Pivot*/
 	TSharedPtr<FUICommandInfo> ResetPivot;
+	/** Toggle Free/Pivot Mode*/
+	TSharedPtr<FUICommandInfo> ToggleFreePivot;
 
 	TMap<FName, TArray<TSharedPtr<FUICommandInfo>>> Commands;
 };
@@ -143,6 +145,12 @@ public:
 
 	// End interfaces
 
+	//If In Pivot Mode, if not then in Free Mode
+	bool IsInPivotMode() const { return bInPivotMode; }
+	//Toggle Pivot Mode
+	void TogglePivotMode() { bInPivotMode = !bInPivotMode; }
+	//Set Pivot Mode
+	void SetPivotMode(bool bVal) { bInPivotMode = bVal; }
 protected:
 
 	UPROPERTY()
@@ -155,9 +163,9 @@ protected:
 	TObjectPtr<UCombinedTransformGizmo> TransformGizmo = nullptr;
 
 protected:
+	bool bInPivotMode = false; //pivot may be in 'free' mode or 'pivot' mode
 	bool bShiftPressedWhenStarted = false;
 	bool bCtrlPressedWhenStarted = false;
-	int32 CtrlModifierId = 1;
 	UWorld* TargetWorld = nullptr;		// target World we will raycast into
 	UInteractiveGizmoManager* GizmoManager = nullptr; //gizmo man
 
@@ -199,6 +207,19 @@ protected:
 	void OnEditorSelectionChanged(UObject* NewSelection);
 	FDelegateHandle OnEditorSelectionChangedHandle;
 
+	// Functions and variables for handling the overlay for switching between free and pivot mode
+	void CreateAndShowPivotOverlay();
+	void RemoveAndDestroyPivotOverlay();
+
+	FMargin GetPivotOverlayPadding() const;
+	
+	static FVector2D LastPivotOverlayLocation;
+	TSharedPtr<SWidget> PivotWidget;
+public:
+	void TryRemovePivotOverlay();
+	void TryShowPivotOverlay();
+	void UpdatePivotOverlayLocation(const FVector2D InLocation, TSharedPtr<IAssetViewport> ActiveLevelViewport);
+
 private:
 	TSharedPtr<FUICommandList> CommandBindings;
 	void ResetPivot();
