NonUnity Fix

diff --git a/Engine/Plugins/Editor/SequencerAnimTools/Source/SequencerAnimTools/Public/SequencerAnimEditPivotTool.h b/Engine/Plugins/Editor/SequencerAnimTools/Source/SequencerAnimTools/Public/SequencerAnimEditPivotTool.h
--- a/Engine/Plugins/Editor/SequencerAnimTools/Source/SequencerAnimTools/Public/SequencerAnimEditPivotTool.h
+++ b/Engine/Plugins/Editor/SequencerAnimTools/Source/SequencerAnimTools/Public/SequencerAnimEditPivotTool.h
@@ -24,6 +24,7 @@ class UControlRig;
 struct FRigControlElement;
 class ISequencer;
 class FUICommandList;
+class IAssetViewport;
 
 /*
 *  The way this sequencer pivot tool works is that
