Sequencer: Animation Tool: Pivot: Little more feedback tweaks

diff --git a/Engine/Plugins/Editor/SequencerAnimTools/Source/SequencerAnimTools/Private/SequencerAnimEditPivotTool.cpp b/Engine/Plugins/Editor/SequencerAnimTools/Source/SequencerAnimTools/Private/SequencerAnimEditPivotTool.cpp
--- a/Engine/Plugins/Editor/SequencerAnimTools/Source/SequencerAnimTools/Private/SequencerAnimEditPivotTool.cpp
+++ b/Engine/Plugins/Editor/SequencerAnimTools/Source/SequencerAnimTools/Private/SequencerAnimEditPivotTool.cpp
@@ -24,14 +24,12 @@
 #include "Tools/MotionTrailOptions.h"
 #include "EditorModeManager.h"
 #include "EditorViewportClient.h"
-#include "EditorModeManager.h"
 #include "Modules/ModuleManager.h"
 #include "ControlRig.h"
 #include "ControlRigSequencerEditorLibrary.h"
 #include "LevelSequence.h"
 #include "LevelSequenceEditorBlueprintLibrary.h"
 #include "ILevelSequenceEditorToolkit.h"
-#include "EditorModeManager.h"
 #include "InteractiveToolManager.h"
 #include "EdModeInteractiveToolsContext.h"
 #include "UnrealEdGlobals.h"
@@ -49,6 +47,8 @@
 #include "Viewports/InViewportUIDragOperation.h"
 #include "Widgets/Layout/SUniformGridPanel.h"
 #include "Styling/StyleColors.h"
+#include "EditorModes.h"
+#include "ControlRigObjectBinding.h"
 
 #define LOCTEXT_NAMESPACE "SequencerAnimTools"
 
@@ -62,6 +62,7 @@ void FEditPivotCommands::RegisterCommands()
 FSavedMappings USequencerPivotTool::SavedPivotLocations;
 FLastSelectedObjects USequencerPivotTool::LastSelectedObjects;
 
+
 static void GetControlRigsAndSequencer(TArray<TWeakObjectPtr<UControlRig>>& ControlRigs, TWeakPtr<ISequencer>& SequencerPtr, ULevelSequence** LevelSequence)
 {
 	*LevelSequence = ULevelSequenceEditorBlueprintLibrary::GetCurrentLevelSequence();
@@ -118,6 +119,121 @@ UInteractiveTool* USequencerPivotToolBuilder::BuildTool(const FToolBuilderState&
 	NewTool->SetWorld(SceneState.World, SceneState.GizmoManager);
 	return NewTool;
 }
+/*
+* ControlRig/Actor Mappings
+*/
+
+
+FTransform FControlRigMappings::GetParentTransform() const
+{
+	if (ControlRig.IsValid())
+	{
+		TSharedPtr<IControlRigObjectBinding> ObjectBinding = ControlRig->GetObjectBinding();
+		if (ObjectBinding.IsValid())
+		{
+			USceneComponent* HostingComponent =  Cast<USceneComponent>(ObjectBinding->GetBoundObject());
+			return HostingComponent ? HostingComponent->GetComponentTransform() : FTransform::Identity;
+		}
+	}
+	return FTransform::Identity;
+}
+
+TOptional<FTransform> FControlRigMappings::GetWorldTransform(const FName& Name) const
+{
+	
+	TOptional<FTransform> WorldTransform;
+	if (ControlRig.IsValid())
+	{
+		const FTransform* LocalTransform = PivotTransforms.Find(Name);
+		if (LocalTransform)
+		{
+			FTransform ToWorldTransform = GetParentTransform();
+			FTransform CurrentTransform = ControlRig.Get()->GetControlGlobalTransform(Name);
+			WorldTransform = (*LocalTransform) * CurrentTransform * ToWorldTransform;
+
+		}
+	}
+	return WorldTransform;
+}
+
+void FControlRigMappings::SetFromWorldTransform(const FName& Name, const FTransform& WorldTransform)
+{
+	if (ControlRig.IsValid())
+	{
+		FTransform ToWorldTransform = GetParentTransform();
+		FTransform GlobalTransform = ControlRig.Get()->GetControlGlobalTransform(Name) * ToWorldTransform;
+		FTransform LocalTransform = WorldTransform.GetRelativeTransform(GlobalTransform);
+		PivotTransforms.Add(Name, LocalTransform);
+	}
+}
+
+TArray<FName> FControlRigMappings::GetAllControls() const
+{
+	TArray<FName> Controls;
+	for (const TPair<FName, FTransform>& Pair : PivotTransforms)
+	{
+		Controls.Add(Pair.Key);
+	}
+	return Controls;
+}
+
+void FControlRigMappings::SelectControls()
+{
+	if (ControlRig.IsValid())
+	{
+		for (TPair<FName, FTransform>& Pair : PivotTransforms)
+		{
+			ControlRig->SelectControl(Pair.Key, true);
+		}
+	}
+}
+
+bool FControlRigMappings::IsAnyControlDeselected() const
+{
+	if (ControlRig.IsValid())
+	{
+		for (const TPair<FName, FTransform>& Pair : PivotTransforms)
+		{
+			if (ControlRig->IsControlSelected(Pair.Key) == false)
+			{
+				return true;
+			}
+		}
+		return false;
+	}
+	return true;
+}
+
+TOptional<FTransform> FActorMappings::GetWorldTransform() const
+{
+	TOptional<FTransform> Transform;
+	if (Actor.IsValid())
+	{
+		FTransform ActorTransform = Actor->ActorToWorld();
+		Transform = PivotTransform * ActorTransform;
+	}
+	return Transform;
+}
+
+void FActorMappings::SetFromWorldTransform( const FTransform& WorldTransform)
+{
+	if (Actor.IsValid())
+	{
+		FTransform ActorTransform = Actor->ActorToWorld();
+		FTransform LocalTransform = WorldTransform.GetRelativeTransform(ActorTransform);
+		PivotTransform = LocalTransform;
+	}
+
+}
+void FActorMappings::SelectActors()
+{
+	if (Actor.IsValid())
+	{
+		GEditor->SelectActor(Actor.Get(), true, true);
+	}
+}
+
+
 /*
 * Proxy
 */
@@ -211,9 +327,30 @@ void USequencerPivotTool::Setup()
 		Commands.ToggleFreePivot,
 		FExecuteAction::CreateUObject(this, &USequencerPivotTool::TogglePivotMode)
 	);
-
+	SetPivotMode(false);
 	CreateAndShowPivotOverlay();
+	
+}
+
+void USequencerPivotTool::SetPivotMode(bool bVal)
+{ 
+	bInPivotMode = bVal;
+	if (bInPivotMode == false)
+	{
+		FEditorModeTools& ModeManager = GLevelEditorModeTools(); //mz need to fix for persona/ animation edit mode
+		ModeManager.SetCoordSystem(COORD_Local);
+	}
+}
+
+void USequencerPivotTool::OnTick(float DeltaTime) 
+{
+	if (bGizmoBeingDragged == false)
+	{
+		SetGizmoBasedOnSelection(true);
+		UpdateGizmoTransform();
+	}
 }
+
 void USequencerPivotTool::ResetPivot()
 {
 	SetGizmoBasedOnSelection(false);
@@ -225,7 +362,6 @@ void USequencerPivotTool::ResetPivot()
 //we use this in case nothing was selected when the tool is active so that we instead select it.
 void USequencerPivotTool::SaveLastSelected()
 {
-
 	LastSelectedObjects.LastSelectedControlRigs.SetNum(0);
 	LastSelectedObjects.LastSelectedActors.SetNum(0);
 
@@ -238,7 +374,7 @@ void USequencerPivotTool::SaveLastSelected()
 			{
 				FControlRigMappings Mapping;
 				Mapping.ControlRig = ControlRig;
-				Mapping.PivotTransforms.Add(Name, GizmoTransform);
+				Mapping.SetFromWorldTransform(Name, GizmoTransform);
 				LastSelectedObjects.LastSelectedControlRigs.Add(Mapping);
 			}
 		}
@@ -249,7 +385,7 @@ void USequencerPivotTool::SaveLastSelected()
 		{
 			FActorMappings Mapping;
 			Mapping.Actor = Actor;
-			Mapping.PivotTransform = GizmoTransform;
+			Mapping.SetFromWorldTransform(GizmoTransform);
 			LastSelectedObjects.LastSelectedActors.Add(Mapping);
 		}
 	}
@@ -266,20 +402,12 @@ void USequencerPivotTool::UpdateTransformAndSelectionOnEntering()
 	{
 		for (FControlRigMappings& Mappings : LastSelectedObjects.LastSelectedControlRigs)
 		{
-			if (Mappings.ControlRig.IsValid())
-			{
-				for (TPair<FName, FTransform>& Pair: Mappings.PivotTransforms)
-				{
-					Mappings.ControlRig->SelectControl(Pair.Key, true);
-				}
-			}
+			Mappings.SelectControls();
 		}
 		for (FActorMappings& ActorMappings : LastSelectedObjects.LastSelectedActors)
 		{
-			if (ActorMappings.Actor.IsValid())
-			{
-				GEditor->SelectActor(ActorMappings.Actor.Get(), true, true);
-			}
+			ActorMappings.SelectActors();
+
 		}
 		SetGizmoBasedOnSelection(!bShiftPressedWhenStarted && !bCtrlPressedWhenStarted);
 	}
@@ -326,9 +454,10 @@ bool USequencerPivotTool::SetGizmoBasedOnSelection(bool bUseSaved)
 			bool bHasSavedPivot = false;
 			if (bUseSaved && Mappings)
 			{
-				if (FTransform* Transform = Mappings->PivotTransforms.Find(Name))
+				TOptional<FTransform>WorldTransform = Mappings->GetWorldTransform(Name);
+				if (WorldTransform.IsSet())
 				{
-					GizmoTransform = *Transform;
+					GizmoTransform = WorldTransform.GetValue();
 					bHasSavedPivot = true;
 				}
 			}
@@ -351,7 +480,17 @@ bool USequencerPivotTool::SetGizmoBasedOnSelection(bool bUseSaved)
 		FActorMappings* Mappings = SavedPivotLocations.ActorMappings.Find(SelectedActor);
 		if (bUseSaved && Mappings)
 		{
-			GizmoTransform = Mappings->PivotTransform;
+
+			TOptional<FTransform>WorldTransform = Mappings->GetWorldTransform();
+			if (WorldTransform.IsSet())
+			{
+				GizmoTransform = WorldTransform.GetValue();
+			}
+			else
+			{
+				GizmoTransform = UControlRigSequencerEditorLibrary::GetActorWorldTransform(LevelSequence, SelectedActor, FrameTime.RoundToFrame(),
+					ESequenceTimeUnit::TickResolution);
+			}
 		}
 		else
 		{
@@ -387,22 +526,13 @@ void USequencerPivotTool::DeactivateMe()
 		{
 			if (Mappings.ControlRig.IsValid())
 			{
-				for (TPair<FName, FTransform>& Pair : Mappings.PivotTransforms)
-				{
-
-					if (Mappings.ControlRig->IsControlSelected(Pair.Key) == false)
-					{
-						bDeactivate = true;
-						break;
-					}
-				}
+				bDeactivate = Mappings.IsAnyControlDeselected();
 			}
 			else
 			{
 				bDeactivate = true;
 				break;
 			}
-			
 		}
 
 		if(bDeactivate == false)
@@ -741,10 +871,10 @@ void USequencerPivotTool::SavePivotTransforms()
 	{
 		FControlRigMappings& Mappings = SavedPivotLocations.ControlRigMappings.FindOrAdd(LastObject.ControlRig);
 		Mappings.ControlRig = LastObject.ControlRig;
-		for (TPair<FName, FTransform>& Pair : LastObject.PivotTransforms)
+		TArray<FName> LastNames = LastObject.GetAllControls();
+		for (const FName& Name : LastNames)
 		{
-			FTransform& Transform = Mappings.PivotTransforms.FindOrAdd(Pair.Key);
-			Transform = GizmoTransform;
+			Mappings.SetFromWorldTransform(Name,GizmoTransform);
 		}
 	}
 
@@ -752,7 +882,7 @@ void USequencerPivotTool::SavePivotTransforms()
 	{
 		FActorMappings& Mappings = SavedPivotLocations.ActorMappings.FindOrAdd(LastObject.Actor);
 		Mappings.Actor = LastObject.Actor;
-		Mappings.PivotTransform = GizmoTransform;
+		Mappings.SetFromWorldTransform(GizmoTransform);
 	}
 }
 
diff --git a/Engine/Plugins/Editor/SequencerAnimTools/Source/SequencerAnimTools/Public/SequencerAnimEditPivotTool.h b/Engine/Plugins/Editor/SequencerAnimTools/Source/SequencerAnimTools/Public/SequencerAnimEditPivotTool.h
--- a/Engine/Plugins/Editor/SequencerAnimTools/Source/SequencerAnimTools/Public/SequencerAnimEditPivotTool.h
+++ b/Engine/Plugins/Editor/SequencerAnimTools/Source/SequencerAnimTools/Public/SequencerAnimEditPivotTool.h
@@ -25,7 +25,6 @@ struct FRigControlElement;
 class ISequencer;
 class FUICommandList;
 class IAssetViewport;
-
 /*
 *  The way this sequencer pivot tool works is that
 *  it will use modify the incoming selections temp pivot
@@ -67,12 +66,25 @@ struct FActorSelectonDuringDrag : public FSelectionDuringDrag
 struct FControlRigMappings
 {
 	TWeakObjectPtr<UControlRig> ControlRig;
+
+	FTransform GetParentTransform() const;
+	TOptional<FTransform> GetWorldTransform(const FName& Name) const;
+	void SetFromWorldTransform(const FName& Name, const FTransform& WorldTransform);
+	TArray<FName> GetAllControls() const;
+	void SelectControls();
+	bool IsAnyControlDeselected() const;
+private:
 	TMap<FName, FTransform> PivotTransforms;
 };
 
 struct FActorMappings
 {
 	TWeakObjectPtr<AActor> Actor;
+
+	TOptional<FTransform> GetWorldTransform() const;
+	void SetFromWorldTransform(const FTransform& WorldTransform);
+	void SelectActors();
+private:
 	FTransform PivotTransform;
 };
 
@@ -136,6 +148,7 @@ public:
 	virtual bool HasCancel() const override { return false; }
 	virtual bool HasAccept() const override { return false; }
 	virtual bool CanAccept() const override { return false; }
+	virtual void OnTick(float DeltaTime) override;
 
 	// IClickBehaviorTarget interface
 	virtual FInputRayHit IsHitByClick(const FInputDeviceRay& ClickPos) override;
@@ -149,9 +162,9 @@ public:
 	//If In Pivot Mode, if not then in Free Mode
 	bool IsInPivotMode() const { return bInPivotMode; }
 	//Toggle Pivot Mode
-	void TogglePivotMode() { bInPivotMode = !bInPivotMode; }
+	void TogglePivotMode() { SetPivotMode(!bInPivotMode); }
 	//Set Pivot Mode
-	void SetPivotMode(bool bVal) { bInPivotMode = bVal; }
+	void SetPivotMode(bool bVal);
 protected:
 
 	UPROPERTY()
