Sequencer: New Motion Trail crash due iterating over float channels expecting there to be all of the channels but there is just 1 for the float channel. Now we just check first for double channels and only check float if ther are no doubles valid. We do this since this code is shared among transform which use doubles and control rig which uses floats.

diff --git a/Engine/Plugins/Editor/SequencerAnimTools/Source/SequencerAnimTools/Private/Sequencer/MotionTrailMovieSceneKey.cpp b/Engine/Plugins/Editor/SequencerAnimTools/Source/SequencerAnimTools/Private/Sequencer/MotionTrailMovieSceneKey.cpp
--- a/Engine/Plugins/Editor/SequencerAnimTools/Source/SequencerAnimTools/Private/Sequencer/MotionTrailMovieSceneKey.cpp
+++ b/Engine/Plugins/Editor/SequencerAnimTools/Source/SequencerAnimTools/Private/Sequencer/MotionTrailMovieSceneKey.cpp
@@ -194,11 +194,12 @@ void FMotionTraiMovieScenelKeyTool::BuildKeys()
 	CachedSelection.Reset();
 
 	UMovieSceneSection* AbsoluteTransformSection = OwningTrail->GetSection();
+	int32 MaxChannel = (int32)ETransformChannel::TranslateZ;//only do the first three channels, 0,1,2 which are position.
 	TArrayView<FMovieSceneDoubleChannel*> DoubleChannels = AbsoluteTransformSection->GetChannelProxy().GetChannels<FMovieSceneDoubleChannel>();
-	if (DoubleChannels.Num() > 0)
+	if (DoubleChannels.Num() > MaxChannel)
 	{
-		DoubleChannels = DoubleChannels.Slice(OwningTrail->GetChannelOffset(), uint8(ETransformChannel::MaxChannel) + 1);
-		for (uint8 ChannelIdx = 0; ChannelIdx <= uint8(ETransformChannel::MaxChannel); ChannelIdx++)
+		DoubleChannels = DoubleChannels.Slice(OwningTrail->GetChannelOffset(), MaxChannel + 1);
+		for (int32 ChannelIdx = 0; ChannelIdx <= MaxChannel; ChannelIdx++)
 		{
 			FMovieSceneDoubleChannel* DoubleChannel = DoubleChannels[ChannelIdx];
 			for (int32 Idx = 0; Idx < DoubleChannel->GetNumKeys(); Idx++)
@@ -213,21 +214,24 @@ void FMotionTraiMovieScenelKeyTool::BuildKeys()
 			}
 		}
 	}
-	TArrayView<FMovieSceneFloatChannel*> FloatChannels = AbsoluteTransformSection->GetChannelProxy().GetChannels<FMovieSceneFloatChannel>();
-	if (FloatChannels.Num() > 0)
+	else
 	{
-		FloatChannels = FloatChannels.Slice(OwningTrail->GetChannelOffset(), uint8(ETransformChannel::MaxChannel) + 1);
-		for (uint8 ChannelIdx = 0; ChannelIdx <= uint8(ETransformChannel::MaxChannel); ChannelIdx++)
+		TArrayView<FMovieSceneFloatChannel*> FloatChannels = AbsoluteTransformSection->GetChannelProxy().GetChannels<FMovieSceneFloatChannel>();
+		if (FloatChannels.Num() > MaxChannel)
 		{
-			FMovieSceneFloatChannel* DoubleChannel = FloatChannels[ChannelIdx];
-			for (int32 Idx = 0; Idx < DoubleChannel->GetNumKeys(); Idx++)
+			FloatChannels = FloatChannels.Slice(OwningTrail->GetChannelOffset(), MaxChannel + 1);
+			for (int32 ChannelIdx = 0; ChannelIdx <= MaxChannel; ChannelIdx++)
 			{
-				FFrameNumber CurTime = DoubleChannel->GetTimes()[Idx];
-
-				if (!Keys.Contains(CurTime))
+				FMovieSceneFloatChannel* FloatChannel = FloatChannels[ChannelIdx];
+				for (int32 Idx = 0; Idx < FloatChannel->GetNumKeys(); Idx++)
 				{
-					TUniquePtr<FTrailKeyInfo> TempKeyInfo = MakeUnique<FTrailKeyInfo>(CurTime, AbsoluteTransformSection, OwningTrail);
-					Keys.Add(MoveTemp(CurTime), MoveTemp(TempKeyInfo));
+					FFrameNumber CurTime = FloatChannel->GetTimes()[Idx];
+
+					if (!Keys.Contains(CurTime))
+					{
+						TUniquePtr<FTrailKeyInfo> TempKeyInfo = MakeUnique<FTrailKeyInfo>(CurTime, AbsoluteTransformSection, OwningTrail);
+						Keys.Add(MoveTemp(CurTime), MoveTemp(TempKeyInfo));
+					}
 				}
 			}
 		}
@@ -265,19 +269,24 @@ void FMotionTraiMovieScenelKeyTool::TranslateSelectedKeys(bool bRight)
 		FFrameNumber Delta = FQualifiedFrameTime(Shift, OwningTrail->GetSequencer()->GetFocusedDisplayRate()).ConvertTo(OwningTrail->GetSequencer()->GetFocusedTickResolution()).RoundToFrame();
 		TRange<FFrameNumber> SectionNewBounds = Section->GetRange();
 		TArrayView<FMovieSceneDoubleChannel*> DoubleChannels = Section->GetChannelProxy().GetChannels<FMovieSceneDoubleChannel>();
-		bool bIsFloat = true;
-		if (DoubleChannels.Num() > 0)
+		TArrayView<FMovieSceneFloatChannel*> FloatChannels = Section->GetChannelProxy().GetChannels<FMovieSceneFloatChannel>();
+		bool bIsFloat = false;
+		int32 MaxChannel = (int32)ETransformChannel::TranslateZ;//only do the first three channels, 0,1,2 which are position.
+		if (DoubleChannels.Num() > MaxChannel)
 		{
 			bIsFloat = false;
-			DoubleChannels = DoubleChannels.Slice(OwningTrail->GetChannelOffset(), uint8(ETransformChannel::MaxChannel) + 1);
+			DoubleChannels = DoubleChannels.Slice(OwningTrail->GetChannelOffset(), MaxChannel + 1);
 		}
-		TArrayView<FMovieSceneFloatChannel*> FloatChannels = Section->GetChannelProxy().GetChannels<FMovieSceneFloatChannel>();
-		if (FloatChannels.Num() > 0)
+		else if (FloatChannels.Num() > MaxChannel)
 		{
 			bIsFloat = true;
-			FloatChannels = FloatChannels.Slice(OwningTrail->GetChannelOffset(), uint8(ETransformChannel::MaxChannel) + 1);
+			FloatChannels = FloatChannels.Slice(OwningTrail->GetChannelOffset(), MaxChannel + 1);
 		}
-		for (uint8 ChannelIdx = 0; ChannelIdx <= uint8(ETransformChannel::MaxChannel); ChannelIdx++)
+		else
+		{
+			return;
+		}
+		for (int32 ChannelIdx = 0; ChannelIdx <= MaxChannel; ChannelIdx++)
 		{
 			TArray<FFrameNumber> KeyTimes;
 			FMovieSceneChannel* Channel = nullptr;
@@ -336,19 +345,25 @@ void FMotionTraiMovieScenelKeyTool::DeleteSelectedKeys()
 	{
 		Section->Modify();
 		bool bIsFloat = true;
+		int32 MaxChannel = (int32)ETransformChannel::TranslateZ;///only do the first three channels, 0,1,2 which are position.
 		TArrayView<FMovieSceneDoubleChannel*> DoubleChannels = Section->GetChannelProxy().GetChannels<FMovieSceneDoubleChannel>();
+		TArrayView<FMovieSceneFloatChannel*> FloatChannels = Section->GetChannelProxy().GetChannels<FMovieSceneFloatChannel>();
+
 		if (DoubleChannels.Num() > 0)
 		{
-			DoubleChannels = DoubleChannels.Slice(OwningTrail->GetChannelOffset(), uint8(ETransformChannel::MaxChannel) + 1);
+			DoubleChannels = DoubleChannels.Slice(OwningTrail->GetChannelOffset(), MaxChannel + 1);
 			bIsFloat = false;
 		}
-		TArrayView<FMovieSceneFloatChannel*> FloatChannels = Section->GetChannelProxy().GetChannels<FMovieSceneFloatChannel>();
-		if (FloatChannels.Num() > 0)
+		else if (FloatChannels.Num() > 0)
 		{
-			FloatChannels = FloatChannels.Slice(OwningTrail->GetChannelOffset(), uint8(ETransformChannel::MaxChannel) + 1);
+			FloatChannels = FloatChannels.Slice(OwningTrail->GetChannelOffset(), MaxChannel + 1);
 			bIsFloat = true;
 		}
-		for (uint8 ChannelIdx = 0; ChannelIdx <= uint8(ETransformChannel::MaxChannel); ChannelIdx++)
+		else
+		{
+			return;
+		}
+		for (int32 ChannelIdx = 0; ChannelIdx <=  MaxChannel; ChannelIdx++)
 		{
 			FMovieSceneChannel* Channel = nullptr;
 			if (bIsFloat)
@@ -492,11 +507,13 @@ TArray<FFrameNumber> FMotionTraiMovieScenelKeyTool::GetTimesFromModifiedTimes(co
 bool FMotionTraiMovieScenelKeyTool::ShouldRebuildKeys()
 {
 	TMap<FFrameNumber, TSet<ETransformChannel>> KeyTimes;
+	int32 MaxChannel = (int32)ETransformChannel::TranslateZ;///only do the first three channels, 0,1,2 which are position.
 	TArrayView<FMovieSceneDoubleChannel*> DoubleChannels = OwningTrail->GetSection()->GetChannelProxy().GetChannels<FMovieSceneDoubleChannel>();
-	if (DoubleChannels.Num() > 0)
+	TArrayView<FMovieSceneFloatChannel*> FloatChannels = OwningTrail->GetSection()->GetChannelProxy().GetChannels<FMovieSceneFloatChannel>();
+	if (DoubleChannels.Num() > MaxChannel)
 	{
-		DoubleChannels = DoubleChannels.Slice(OwningTrail->GetChannelOffset(), uint8(ETransformChannel::MaxChannel) + 1);
-		for (uint8 ChannelIdx = 0; ChannelIdx <= uint8(ETransformChannel::MaxChannel); ChannelIdx++)
+		DoubleChannels = DoubleChannels.Slice(OwningTrail->GetChannelOffset(), MaxChannel + 1);
+		for (int32 ChannelIdx = 0; ChannelIdx <= MaxChannel; ChannelIdx++)
 		{
 			FMovieSceneDoubleChannel* DoubleChannel = DoubleChannels[ChannelIdx];
 			for (int32 Idx = 0; Idx < DoubleChannel->GetNumKeys(); Idx++)
@@ -506,11 +523,10 @@ bool FMotionTraiMovieScenelKeyTool::ShouldRebuildKeys()
 			}
 		}
 	}
-	TArrayView<FMovieSceneFloatChannel*> FloatChannels = OwningTrail->GetSection()->GetChannelProxy().GetChannels<FMovieSceneFloatChannel>();
-	if (FloatChannels.Num() > 0)
+	else if (FloatChannels.Num() > MaxChannel)
 	{
-		FloatChannels = FloatChannels.Slice(OwningTrail->GetChannelOffset(), uint8(ETransformChannel::MaxChannel) + 1);
-		for (uint8 ChannelIdx = 0; ChannelIdx <= uint8(ETransformChannel::MaxChannel); ChannelIdx++)
+		FloatChannels = FloatChannels.Slice(OwningTrail->GetChannelOffset(), MaxChannel + 1);
+		for (int32 ChannelIdx = 0; ChannelIdx <= MaxChannel; ChannelIdx++)
 		{
 			FMovieSceneFloatChannel* FloatChannel = FloatChannels[ChannelIdx];
 			for (int32 Idx = 0; Idx < FloatChannel->GetNumKeys(); Idx++)
@@ -532,22 +548,22 @@ bool FMotionTraiMovieScenelKeyTool::ShouldRebuildKeys()
 		{
 			return true;
 		}
-		for (uint8 ChannelIdx = 0; ChannelIdx <= uint8(ETransformChannel::MaxChannel); ChannelIdx++)
+		for (int32 ChannelIdx = 0; ChannelIdx <= MaxChannel; ChannelIdx++)
 		{
 			const ETransformChannel TransformChannel = ETransformChannel(ChannelIdx);
-			if (DoubleChannels.Num() > 0)
+			if (DoubleChannels.Num() > MaxChannel)
 			{
 				if ((!TimeKeyPair.Value.Contains(TransformChannel) && Keys[TimeKeyPair.Key]->IdxMap.Contains(TransformChannel)) ||
-					(Keys[TimeKeyPair.Key]->IdxMap.Contains(TransformChannel) && DoubleChannels[uint8(TransformChannel)]->GetData().GetIndex(Keys[TimeKeyPair.Key]->IdxMap[TransformChannel]) == INDEX_NONE) ||
+					(Keys[TimeKeyPair.Key]->IdxMap.Contains(TransformChannel) && DoubleChannels[int32(TransformChannel)]->GetData().GetIndex(Keys[TimeKeyPair.Key]->IdxMap[TransformChannel]) == INDEX_NONE) ||
 					(TimeKeyPair.Value.Contains(TransformChannel) && !Keys[TimeKeyPair.Key]->IdxMap.Contains(TransformChannel)))
 				{
 					return true;
 				}
 			}
-			if (FloatChannels.Num() > 0)
+			else if (FloatChannels.Num() > MaxChannel)
 			{
 				if ((!TimeKeyPair.Value.Contains(TransformChannel) && Keys[TimeKeyPair.Key]->IdxMap.Contains(TransformChannel)) ||
-					(Keys[TimeKeyPair.Key]->IdxMap.Contains(TransformChannel) && FloatChannels[uint8(TransformChannel)]->GetData().GetIndex(Keys[TimeKeyPair.Key]->IdxMap[TransformChannel]) == INDEX_NONE) ||
+					(Keys[TimeKeyPair.Key]->IdxMap.Contains(TransformChannel) && FloatChannels[int32(TransformChannel)]->GetData().GetIndex(Keys[TimeKeyPair.Key]->IdxMap[TransformChannel]) == INDEX_NONE) ||
 					(TimeKeyPair.Value.Contains(TransformChannel) && !Keys[TimeKeyPair.Key]->IdxMap.Contains(TransformChannel)))
 				{
 					return true;
@@ -589,36 +605,35 @@ FTrailKeyInfo::FTrailKeyInfo(const FFrameNumber InFrameNumber, UMovieSceneSectio
 , bDirty(true)
 , OwningTrail(InOwningTrail)
 {
+	int32 MaxChannel = (int32)ETransformChannel::TranslateZ;///only do the first three channels, 0,1,2 which are position.
+	TArrayView<FMovieSceneDoubleChannel*> DoubleChannels = InSection->GetChannelProxy().GetChannels<FMovieSceneDoubleChannel>();
+	TArrayView<FMovieSceneFloatChannel*> FloatChannels = InSection->GetChannelProxy().GetChannels<FMovieSceneFloatChannel>();
+
+	if (DoubleChannels.Num() > MaxChannel)
 	{
-		TArrayView<FMovieSceneDoubleChannel*> Channels = InSection->GetChannelProxy().GetChannels<FMovieSceneDoubleChannel>();
-		if (Channels.Num() > 0)
+		DoubleChannels = DoubleChannels.Slice(OwningTrail->GetChannelOffset(),  MaxChannel + 1);
+		for (int32 Idx = 0; Idx <= MaxChannel; Idx++)
 		{
-			Channels = Channels.Slice(OwningTrail->GetChannelOffset(), uint8(ETransformChannel::MaxChannel) + 1);
-			for (uint8 Idx = 0; Idx <= uint8(ETransformChannel::MaxChannel); Idx++)
+			const int32 FoundIdx = DoubleChannels[Idx]->GetData().FindKey(InFrameNumber);
+			if (FoundIdx != INDEX_NONE)
 			{
-				const int32 FoundIdx = Channels[Idx]->GetData().FindKey(InFrameNumber);
-				if (FoundIdx != INDEX_NONE)
-				{
-					IdxMap.Add(ETransformChannel(Idx), Channels[Idx]->GetData().GetHandle(FoundIdx));
-				}
+				IdxMap.Add(ETransformChannel(Idx), DoubleChannels[Idx]->GetData().GetHandle(FoundIdx));
 			}
 		}
 	}
+	else if (FloatChannels.Num() > MaxChannel)
 	{
-		TArrayView<FMovieSceneFloatChannel*> Channels = InSection->GetChannelProxy().GetChannels<FMovieSceneFloatChannel>();
-		if (Channels.Num() > 0)
+		FloatChannels = FloatChannels.Slice(OwningTrail->GetChannelOffset(), MaxChannel + 1);
+		for (int32 Idx = 0; Idx <= MaxChannel; Idx++)
 		{
-			Channels = Channels.Slice(OwningTrail->GetChannelOffset(), uint8(ETransformChannel::MaxChannel) + 1);
-			for (uint8 Idx = 0; Idx <= uint8(ETransformChannel::MaxChannel); Idx++)
+			const int32 FoundIdx = FloatChannels[Idx]->GetData().FindKey(InFrameNumber);
+			if (FoundIdx != INDEX_NONE)
 			{
-				const int32 FoundIdx = Channels[Idx]->GetData().FindKey(InFrameNumber);
-				if (FoundIdx != INDEX_NONE)
-				{
-					IdxMap.Add(ETransformChannel(Idx), Channels[Idx]->GetData().GetHandle(FoundIdx));
-				}
+				IdxMap.Add(ETransformChannel(Idx), FloatChannels[Idx]->GetData().GetHandle(FoundIdx));
 			}
 		}
 	}
+	
 }
 
 void FTrailKeyInfo::UpdateKeyTransform(EGetKeyFrom UpdateType)
diff --git a/Engine/Plugins/Editor/SequencerAnimTools/Source/SequencerAnimTools/Private/Sequencer/MovieSceneTransformTrail.cpp b/Engine/Plugins/Editor/SequencerAnimTools/Source/SequencerAnimTools/Private/Sequencer/MovieSceneTransformTrail.cpp
--- a/Engine/Plugins/Editor/SequencerAnimTools/Source/SequencerAnimTools/Private/Sequencer/MovieSceneTransformTrail.cpp
+++ b/Engine/Plugins/Editor/SequencerAnimTools/Source/SequencerAnimTools/Private/Sequencer/MovieSceneTransformTrail.cpp
@@ -302,26 +302,24 @@ TRange<double> FMovieSceneTransformTrail::GetEffectiveSectionRange(int32 Channel
 
 	TRange<FFrameNumber> EffectiveRange = TRange<FFrameNumber>::Empty();
 
+	int32 MaxChannel = (int32)ETransformChannel::TranslateZ; //0,1,2 are the position channels only onces we care about.
+	TArrayView<FMovieSceneDoubleChannel*> DoubleChannels = TransformSection->GetChannelProxy().GetChannels<FMovieSceneDoubleChannel>();
+	TArrayView<FMovieSceneFloatChannel*> FloatChannels = TransformSection->GetChannelProxy().GetChannels<FMovieSceneFloatChannel>();
+
+	if (DoubleChannels.Num() > MaxChannel)
 	{
-		TArrayView<FMovieSceneDoubleChannel*> Channels = TransformSection->GetChannelProxy().GetChannels<FMovieSceneDoubleChannel>();
-		if (Channels.Num() > 0)
+		for (int32 ChannelIdx = ChannelOffset; ChannelIdx <= (ChannelOffset + MaxChannel); ChannelIdx++)
 		{
-			for (int32 ChannelIdx = ChannelOffset; ChannelIdx <= ChannelOffset + uint8(ETransformChannel::MaxChannel); ChannelIdx++)
-			{
-				FMovieSceneDoubleChannel* Channel = Channels[ChannelIdx];
-				EffectiveRange = TRange<FFrameNumber>::Hull(EffectiveRange, Channel->ComputeEffectiveRange());
-			}
+			FMovieSceneDoubleChannel* Channel = DoubleChannels[ChannelIdx];
+			EffectiveRange = TRange<FFrameNumber>::Hull(EffectiveRange, Channel->ComputeEffectiveRange());
 		}
 	}
+	if (FloatChannels.Num() > MaxChannel)
 	{
-		TArrayView<FMovieSceneFloatChannel*> Channels = TransformSection->GetChannelProxy().GetChannels<FMovieSceneFloatChannel>();
-		if (Channels.Num() > 0)
+		for (int32 ChannelIdx = ChannelOffset; ChannelIdx <= (ChannelOffset + MaxChannel); ChannelIdx++)
 		{
-			for (int32 ChannelIdx = ChannelOffset; ChannelIdx <= ChannelOffset + uint8(ETransformChannel::MaxChannel); ChannelIdx++)
-			{
-				FMovieSceneFloatChannel* Channel = Channels[ChannelIdx];
-				EffectiveRange = TRange<FFrameNumber>::Hull(EffectiveRange, Channel->ComputeEffectiveRange());
-			}
+			FMovieSceneFloatChannel* Channel = FloatChannels[ChannelIdx];
+			EffectiveRange = TRange<FFrameNumber>::Hull(EffectiveRange, Channel->ComputeEffectiveRange());
 		}
 	}
 
@@ -600,24 +598,6 @@ FMovieSceneControlRigTransformTrail::FMovieSceneControlRigTransformTrail(USceneC
 {
 }
 
-FTransform FMovieSceneControlRigTransformTrail::EvaluateChannelsAtTime(TArrayView<FMovieSceneFloatChannel*> Channels, FFrameTime Time) const
-{
-	FVector3f TempTranslation;
-	Channels[0]->Evaluate(Time, TempTranslation.X);
-	Channels[1]->Evaluate(Time, TempTranslation.Y);
-	Channels[2]->Evaluate(Time, TempTranslation.Z);
-	FRotator3f TempRotation;
-	Channels[3]->Evaluate(Time, TempRotation.Roll);
-	Channels[4]->Evaluate(Time, TempRotation.Pitch);
-	Channels[5]->Evaluate(Time, TempRotation.Yaw);
-	FVector3f TempScale3D;
-	Channels[6]->Evaluate(Time, TempScale3D.X);
-	Channels[7]->Evaluate(Time, TempScale3D.Y);
-	Channels[8]->Evaluate(Time, TempScale3D.Z);
-	FTransform TempTransform = FTransform(FRotator(TempRotation), FVector(TempTranslation), FVector(TempScale3D));
-	TempTransform.NormalizeRotation();
-	return TempTransform;
-}
 
 void FMovieSceneControlRigTransformTrail::UpdateCacheTimes(const FSceneContext& InSceneContext)
 {
diff --git a/Engine/Plugins/Editor/SequencerAnimTools/Source/SequencerAnimTools/Private/Sequencer/MovieSceneTransformTrail.h b/Engine/Plugins/Editor/SequencerAnimTools/Source/SequencerAnimTools/Private/Sequencer/MovieSceneTransformTrail.h
--- a/Engine/Plugins/Editor/SequencerAnimTools/Source/SequencerAnimTools/Private/Sequencer/MovieSceneTransformTrail.h
+++ b/Engine/Plugins/Editor/SequencerAnimTools/Source/SequencerAnimTools/Private/Sequencer/MovieSceneTransformTrail.h
@@ -130,7 +130,6 @@ protected:
 
 private:
 	bool bUseKeysForTrajectory = false; //on when interatively moving
-	FTransform EvaluateChannelsAtTime(TArrayView<FMovieSceneFloatChannel*> Channels, FFrameTime Time) const;
 	FName ControlName;
 };
 
