Sequencer: Pivot Tool: Trivial Fix to get selected actors correctly when re-entering the tool.
#trivial

diff --git a/Engine/Plugins/Editor/SequencerAnimTools/Source/SequencerAnimTools/Private/SequencerAnimEditPivotTool.cpp b/Engine/Plugins/Editor/SequencerAnimTools/Source/SequencerAnimTools/Private/SequencerAnimEditPivotTool.cpp
--- a/Engine/Plugins/Editor/SequencerAnimTools/Source/SequencerAnimTools/Private/SequencerAnimEditPivotTool.cpp
+++ b/Engine/Plugins/Editor/SequencerAnimTools/Source/SequencerAnimTools/Private/SequencerAnimEditPivotTool.cpp
@@ -230,6 +230,7 @@ void FActorMappings::SelectActors()
 	if (Actor.IsValid())
 	{
 		GEditor->SelectActor(Actor.Get(), true, true);
+		GEditor->NoteSelectionChange();
 	}
 }
 
@@ -289,6 +290,8 @@ void USequencerPivotTool::Setup()
 
 	GetControlRigsAndSequencer(ControlRigs, SequencerPtr, &LevelSequence);
 
+	UpdateTransformAndSelectionOnEntering();
+
 	TArray<AActor*> SelectedActors;
 	GEditor->GetSelectedActors()->GetSelectedObjects(SelectedActors);
 	Actors.SetNum(0);
@@ -297,7 +300,6 @@ void USequencerPivotTool::Setup()
 		Actors.Add(Actor);
 	}
 
-	UpdateTransformAndSelectionOnEntering();
 	UpdateGizmoTransform();
 	UpdateGizmoVisibility();
 	//we get delegates last since we may select something above
