Sequencer: Avoid deregistering gizmo helper to avoid a crash bug that results from Enter() preceding Exit() when switching to another mode and using gizmos there (esp modeling tools).

diff --git a/Engine/Plugins/Editor/SequencerAnimTools/Source/SequencerAnimTools/Private/SequencerToolsEditMode.cpp b/Engine/Plugins/Editor/SequencerAnimTools/Source/SequencerAnimTools/Private/SequencerToolsEditMode.cpp
--- a/Engine/Plugins/Editor/SequencerAnimTools/Source/SequencerAnimTools/Private/SequencerToolsEditMode.cpp
+++ b/Engine/Plugins/Editor/SequencerAnimTools/Source/SequencerAnimTools/Private/SequencerToolsEditMode.cpp
@@ -65,7 +65,9 @@ void USequencerToolsEditMode::Exit()
 			LevelEditorPtr.Pin()->GetEditorModeManager().GetInteractiveToolsContext()->ToolManager->UnregisterToolType(TEXT("SequencerMotionTrail"));
 			LevelEditorPtr.Pin()->GetEditorModeManager().GetInteractiveToolsContext()->ToolManager->UnregisterToolType(TEXT("SequencerPivotTool"));
 
-			UE::TransformGizmoUtil::DeregisterTransformGizmoContextObject(LevelEditorPtr.Pin()->GetEditorModeManager().GetInteractiveToolsContext());
+			// TODO: cannot deregister here due to a bug where another mode's Enter() is called before our Exit() on mode switch, and we end
+			// up deregistering a helper that the other mode was relying on. Could uncomment the below when that bug is fixed.
+			//UE::TransformGizmoUtil::DeregisterTransformGizmoContextObject(LevelEditorPtr.Pin()->GetEditorModeManager().GetInteractiveToolsContext());
 		}
 	}
 }
