Sequencer: Fix crash in other modes that use the same gizmo context object by making sure that sequencer calls the proper super class enter/exit functions and uses the mode-level context store for its version.

diff --git a/Engine/Plugins/Editor/SequencerAnimTools/Source/SequencerAnimTools/Private/SequencerToolsEditMode.cpp b/Engine/Plugins/Editor/SequencerAnimTools/Source/SequencerAnimTools/Private/SequencerToolsEditMode.cpp
--- a/Engine/Plugins/Editor/SequencerAnimTools/Source/SequencerAnimTools/Private/SequencerToolsEditMode.cpp
+++ b/Engine/Plugins/Editor/SequencerAnimTools/Source/SequencerAnimTools/Private/SequencerToolsEditMode.cpp
@@ -36,6 +36,8 @@ USequencerToolsEditMode::~USequencerToolsEditMode()
 
 void USequencerToolsEditMode::Enter()
 {
+	Super::Enter();
+
 	FLevelEditorModule* LevelEditorModule = FModuleManager::GetModulePtr<FLevelEditorModule>(TEXT("LevelEditor"));
 	
 	if (LevelEditorModule)
@@ -47,7 +49,9 @@ void USequencerToolsEditMode::Enter()
 			LevelEditorPtr.Pin()->GetEditorModeManager().GetInteractiveToolsContext()->ToolManager->RegisterToolType(TEXT("SequencerMotionTrail"), NewObject<UMotionTrailToolBuilder>(this));
 			LevelEditorPtr.Pin()->GetEditorModeManager().GetInteractiveToolsContext()->ToolManager->RegisterToolType(TEXT("SequencerPivotTool"), NewObject<USequencerPivotToolBuilder>(this));
 
-			UE::TransformGizmoUtil::RegisterTransformGizmoContextObject(LevelEditorPtr.Pin()->GetEditorModeManager().GetInteractiveToolsContext());
+			// Currently there are some issues when the gizmo context object is shared among modes, so make sure to
+			// register it in the mode-level context store so it's not picked up by other modes.
+			UE::TransformGizmoUtil::RegisterTransformGizmoContextObject(GetInteractiveToolsContext(EToolsContextScope::EdMode));
 		}
 	}
 }
@@ -65,11 +69,11 @@ void USequencerToolsEditMode::Exit()
 			LevelEditorPtr.Pin()->GetEditorModeManager().GetInteractiveToolsContext()->ToolManager->UnregisterToolType(TEXT("SequencerMotionTrail"));
 			LevelEditorPtr.Pin()->GetEditorModeManager().GetInteractiveToolsContext()->ToolManager->UnregisterToolType(TEXT("SequencerPivotTool"));
 
-			// TODO: cannot deregister here due to a bug where another mode's Enter() is called before our Exit() on mode switch, and we end
-			// up deregistering a helper that the other mode was relying on. Could uncomment the below when that bug is fixed.
-			//UE::TransformGizmoUtil::DeregisterTransformGizmoContextObject(LevelEditorPtr.Pin()->GetEditorModeManager().GetInteractiveToolsContext());
+			UE::TransformGizmoUtil::DeregisterTransformGizmoContextObject(GetInteractiveToolsContext(EToolsContextScope::EdMode));
 		}
 	}
+
+	Super::Exit();
 }
 
 bool USequencerToolsEditMode::IsCompatibleWith(FEditorModeID OtherModeID) const
