Sequencer:Motion Trail: When reevaluate a trail make sure to alo recalucate it's keys.
#trivial

diff --git a/Engine/Plugins/Editor/SequencerAnimTools/Source/SequencerAnimTools/Private/Sequencer/MovieSceneTransformTrail.cpp b/Engine/Plugins/Editor/SequencerAnimTools/Source/SequencerAnimTools/Private/Sequencer/MovieSceneTransformTrail.cpp
--- a/Engine/Plugins/Editor/SequencerAnimTools/Source/SequencerAnimTools/Private/Sequencer/MovieSceneTransformTrail.cpp
+++ b/Engine/Plugins/Editor/SequencerAnimTools/Source/SequencerAnimTools/Private/Sequencer/MovieSceneTransformTrail.cpp
@@ -105,8 +105,16 @@ void FMovieSceneTransformTrail::CalculateEditedTimes(const FTrailHierarchy* Trai
 		}
 	}
 }
+
+void FMovieSceneTransformTrail::ForceEvaluateNextTick()
+{
+	FTrail::ForceEvaluateNextTick();
+	KeyTool->DirtyKeyTransforms();
+}
+
 void FMovieSceneTransformTrail::UpdateKeysInRange(const TRange<double>& ViewRange)
 {
+	KeyTool->DirtyKeyTransforms();
 	KeyTool->UpdateKeysInRange(ViewRange);
 }
 
diff --git a/Engine/Plugins/Editor/SequencerAnimTools/Source/SequencerAnimTools/Private/Sequencer/MovieSceneTransformTrail.h b/Engine/Plugins/Editor/SequencerAnimTools/Source/SequencerAnimTools/Private/Sequencer/MovieSceneTransformTrail.h
--- a/Engine/Plugins/Editor/SequencerAnimTools/Source/SequencerAnimTools/Private/Sequencer/MovieSceneTransformTrail.h
+++ b/Engine/Plugins/Editor/SequencerAnimTools/Source/SequencerAnimTools/Private/Sequencer/MovieSceneTransformTrail.h
@@ -45,6 +45,7 @@ public:
 	virtual void UpdateKeysInRange(const TRange<double>& ViewRange) override;
 	virtual TArray<FFrameNumber> GetKeyTimes() const override;
 	virtual TArray<FFrameNumber> GetSelectedKeyTimes() const override;
+	virtual void ForceEvaluateNextTick() override;
 	// End FTrail interface
 
 	TSharedPtr<ISequencer> GetSequencer() const { return WeakSequencer.Pin(); }
diff --git a/Engine/Plugins/Editor/SequencerAnimTools/Source/SequencerAnimTools/Public/Trail.h b/Engine/Plugins/Editor/SequencerAnimTools/Source/SequencerAnimTools/Public/Trail.h
--- a/Engine/Plugins/Editor/SequencerAnimTools/Source/SequencerAnimTools/Public/Trail.h
+++ b/Engine/Plugins/Editor/SequencerAnimTools/Source/SequencerAnimTools/Public/Trail.h
@@ -132,8 +132,10 @@ public:
 		}
 	}
 
+	virtual void ForceEvaluateNextTick() { bForceEvaluateNextTick = true; }
+
+
 	FTrajectoryDrawInfo* GetDrawInfo() { return DrawInfo.Get(); }
-	void ForceEvaluateNextTick() { bForceEvaluateNextTick = true; }
 
 	ETrailCacheState GetCacheState() const { return CacheState; }
 protected:
