Sequencer: Motion Trail: On reconstruct make sure we add channels for absolutes but still handle if we have a bad rig that has a control that doesn't have a parameter.

diff --git a/Engine/Plugins/Animation/ControlRig/Source/ControlRig/Private/Sequencer/MovieSceneControlRigParameterTrack.cpp b/Engine/Plugins/Animation/ControlRig/Source/ControlRig/Private/Sequencer/MovieSceneControlRigParameterTrack.cpp
--- a/Engine/Plugins/Animation/ControlRig/Source/ControlRig/Private/Sequencer/MovieSceneControlRigParameterTrack.cpp
+++ b/Engine/Plugins/Animation/ControlRig/Source/ControlRig/Private/Sequencer/MovieSceneControlRigParameterTrack.cpp
@@ -381,7 +381,7 @@ void UMovieSceneControlRigParameterTrack::ReconstructControlRig()
 					{
 						CRSection->SpaceChannelAdded().AddUObject(this, &UMovieSceneControlRigParameterTrack::HandleOnSpaceAdded);
 					}
-					CRSection->RecreateWithThisControlRig(ControlRig, false);
+					CRSection->RecreateWithThisControlRig(ControlRig, CRSection->GetBlendType() == EMovieSceneBlendType::Absolute);
 				}
 			}
 		}
diff --git a/Engine/Plugins/Editor/SequencerAnimTools/Source/SequencerAnimTools/Private/Sequencer/MotionTrailMovieSceneKey.cpp b/Engine/Plugins/Editor/SequencerAnimTools/Source/SequencerAnimTools/Private/Sequencer/MotionTrailMovieSceneKey.cpp
--- a/Engine/Plugins/Editor/SequencerAnimTools/Source/SequencerAnimTools/Private/Sequencer/MotionTrailMovieSceneKey.cpp
+++ b/Engine/Plugins/Editor/SequencerAnimTools/Source/SequencerAnimTools/Private/Sequencer/MotionTrailMovieSceneKey.cpp
@@ -192,7 +192,10 @@ void FMotionTraiMovieScenelKeyTool::BuildKeys()
 	TArray<FFrameNumber> KeyTimes = SelectedKeyTimes();
 	Keys.Reset();
 	CachedSelection.Reset();
-
+	if (OwningTrail->GetChannelOffset() == INDEX_NONE)
+	{
+		return;
+	}
 	UMovieSceneSection* AbsoluteTransformSection = OwningTrail->GetSection();
 	int32 MaxChannel = (int32)ETransformChannel::TranslateZ;//only do the first three channels, 0,1,2 which are position.
 	TArrayView<FMovieSceneDoubleChannel*> DoubleChannels = AbsoluteTransformSection->GetChannelProxy().GetChannels<FMovieSceneDoubleChannel>();
@@ -261,6 +264,10 @@ TArray<FKeyHandle> FMotionTraiMovieScenelKeyTool::GetSelectedKeyHandles(FMovieSc
 
 void FMotionTraiMovieScenelKeyTool::TranslateSelectedKeys(bool bRight)
 {
+	if (OwningTrail->GetChannelOffset() == INDEX_NONE)
+	{
+		return;
+	}
 	UMovieSceneSection* Section = OwningTrail->GetSection();
 	if (CachedSelection.Num() > 0 && Section && Section->TryModify())
 	{
@@ -340,6 +347,10 @@ void FMotionTraiMovieScenelKeyTool::TranslateSelectedKeys(bool bRight)
 
 void FMotionTraiMovieScenelKeyTool::DeleteSelectedKeys()
 {
+	if (OwningTrail->GetChannelOffset() == INDEX_NONE)
+	{
+		return;
+	}
 	UMovieSceneSection* Section = OwningTrail->GetSection();
 	if (CachedSelection.Num() > 0 && Section && Section->TryModify())
 	{
@@ -506,6 +517,10 @@ TArray<FFrameNumber> FMotionTraiMovieScenelKeyTool::GetTimesFromModifiedTimes(co
 
 bool FMotionTraiMovieScenelKeyTool::ShouldRebuildKeys()
 {
+	if (OwningTrail->GetChannelOffset() == INDEX_NONE)
+	{
+		return false;
+	}
 	TMap<FFrameNumber, TSet<ETransformChannel>> KeyTimes;
 	int32 MaxChannel = (int32)ETransformChannel::TranslateZ;///only do the first three channels, 0,1,2 which are position.
 	TArrayView<FMovieSceneDoubleChannel*> DoubleChannels = OwningTrail->GetSection()->GetChannelProxy().GetChannels<FMovieSceneDoubleChannel>();
@@ -605,6 +620,10 @@ FTrailKeyInfo::FTrailKeyInfo(const FFrameNumber InFrameNumber, UMovieSceneSectio
 , bDirty(true)
 , OwningTrail(InOwningTrail)
 {
+	if (OwningTrail->GetChannelOffset() == INDEX_NONE)
+	{
+		return;
+	}
 	int32 MaxChannel = (int32)ETransformChannel::TranslateZ;///only do the first three channels, 0,1,2 which are position.
 	TArrayView<FMovieSceneDoubleChannel*> DoubleChannels = InSection->GetChannelProxy().GetChannels<FMovieSceneDoubleChannel>();
 	TArrayView<FMovieSceneFloatChannel*> FloatChannels = InSection->GetChannelProxy().GetChannels<FMovieSceneFloatChannel>();
diff --git a/Engine/Plugins/Editor/SequencerAnimTools/Source/SequencerAnimTools/Private/Sequencer/MovieSceneTransformTrail.cpp b/Engine/Plugins/Editor/SequencerAnimTools/Source/SequencerAnimTools/Private/Sequencer/MovieSceneTransformTrail.cpp
--- a/Engine/Plugins/Editor/SequencerAnimTools/Source/SequencerAnimTools/Private/Sequencer/MovieSceneTransformTrail.cpp
+++ b/Engine/Plugins/Editor/SequencerAnimTools/Source/SequencerAnimTools/Private/Sequencer/MovieSceneTransformTrail.cpp
@@ -721,9 +721,9 @@ int32 FMovieSceneControlRigTransformTrail::GetChannelOffset() const
 	if (CRParamSection)
 	{
 		FChannelMapInfo* pChannelIndex = CRParamSection->ControlChannelMap.Find(ControlName);
-		return pChannelIndex->ChannelIndex;;
+		return pChannelIndex ? pChannelIndex->ChannelIndex : INDEX_NONE;
 	}
-	return 0;
+	return INDEX_NONE;
 }
 
 bool FMovieSceneControlRigTransformTrail::ApplyDelta(const FVector& Pos, const FRotator& Rot, const FVector& WidgetLocation)
