Sequencer: Pivot Tool need to handle moving out of non zero rotation and space and handle non zero parent space. Added rotations to the pivot tool late and didn't test rotations with the parent rotated.

diff --git a/Engine/Plugins/Editor/SequencerAnimTools/Source/SequencerAnimTools/Private/SequencerAnimEditPivotTool.cpp b/Engine/Plugins/Editor/SequencerAnimTools/Source/SequencerAnimTools/Private/SequencerAnimEditPivotTool.cpp
--- a/Engine/Plugins/Editor/SequencerAnimTools/Source/SequencerAnimTools/Private/SequencerAnimEditPivotTool.cpp
+++ b/Engine/Plugins/Editor/SequencerAnimTools/Source/SequencerAnimTools/Private/SequencerAnimEditPivotTool.cpp
@@ -124,8 +124,6 @@ UInteractiveTool* USequencerPivotToolBuilder::BuildTool(const FToolBuilderState&
 /*
 * ControlRig/Actor Mappings
 */
-
-
 FTransform FControlRigMappings::GetParentTransform() const
 {
 	if (ControlRig.IsValid())
@@ -772,14 +770,21 @@ void USequencerPivotTool::GizmoTransformChanged(UTransformProxy* Proxy, FTransfo
 						FVector RotatedDiff = Diff.GetRotation().RotateVector(LocDiff);
 						FVector NewLocation = StartDragTransform.TransformPosition(RotatedDiff);
 						ControlDrag.CurrentTransform.SetLocation(NewLocation);
-						if (bAlsoCalcRotation)
-						{
-							FQuat Quat = ControlDrag.CurrentTransform.GetRotation() * Diff.GetRotation();
-							ControlDrag.CurrentTransform.SetRotation(Quat);
-						}
-						UControlRigSequencerEditorLibrary::SetControlRigWorldTransform(ControlDrag.LevelSequence, ControlDrag.ControlRig, ControlDrag.ControlName,
-							ControlDrag.CurrentFrame, ControlDrag.CurrentTransform, ESequenceTimeUnit::TickResolution, bSetKey);
 					}
+					if (bAlsoCalcRotation)
+					{
+						FQuat StartDragRot = StartDragTransform.GetRotation();
+						FQuat ActorDragRot = ControlDrag.CurrentTransform.GetRotation().Inverse();
+						FQuat ParentRot = ActorDragRot * StartDragRot;
+						ParentRot = ParentRot.Inverse();
+						FQuat CurrentDragRot = Transform.GetRotation();
+						FQuat OptQuat = CurrentDragRot * ParentRot;
+						OptQuat.EnforceShortestArcWith(ControlDrag.CurrentTransform.GetRotation());
+						ControlDrag.CurrentTransform.SetRotation(OptQuat);
+					}
+					UControlRigSequencerEditorLibrary::SetControlRigWorldTransform(ControlDrag.LevelSequence, ControlDrag.ControlRig, ControlDrag.ControlName,
+						ControlDrag.CurrentFrame, ControlDrag.CurrentTransform, ESequenceTimeUnit::TickResolution, bSetKey);
+					
 				}
 				else if (bTranslationChanged)
 				{
@@ -819,15 +824,27 @@ void USequencerPivotTool::GizmoTransformChanged(UTransformProxy* Proxy, FTransfo
 						TOptional <FQuat> OptQuat;
 						if (bAlsoCalcRotation)
 						{
-							OptQuat = ActorDrag.CurrentTransform.GetRotation() * Diff.GetRotation();
+							FQuat StartDragRot = StartDragTransform.GetRotation();
+							FQuat ActorDragRot = ActorDrag.CurrentTransform.GetRotation().Inverse();
+							FQuat ParentRot = ActorDragRot * StartDragRot;
+							ParentRot = ParentRot.Inverse();
+							FQuat CurrentDragRot = Transform.GetRotation();
+							OptQuat = CurrentDragRot * ParentRot;
 							ActorDrag.CurrentTransform.SetRotation(OptQuat.GetValue());
 						}
 						SetLocation(ActorDrag.Actor, NewLocation, OptQuat);
-						/*Note that calling the bit below isn't enought. That's just sends a begin/end event, the other bits are needed, and needed between these events
-						UEditorActorSubsystem* EditorActorSubsystem = GEditor->GetEditorSubsystem<UEditorActorSubsystem>();
-						EditorActorSubsystem->SetActorTransform(ActorDrag.Actor,ActorDrag.CurrentTransform);
-						GUnrealEd->UpdatePivotLocationForSelection();
-						*/
+					}
+					else if (bAlsoCalcRotation)
+					{
+						FQuat StartDragRot = StartDragTransform.GetRotation();
+						FQuat ActorDragRot = ActorDrag.CurrentTransform.GetRotation().Inverse();
+						FQuat ParentRot =  ActorDragRot * StartDragRot;
+						ParentRot = ParentRot.Inverse();
+						FQuat CurrentDragRot = Transform.GetRotation();
+						FQuat OptQuat = CurrentDragRot * ParentRot;
+						OptQuat.EnforceShortestArcWith(ActorDrag.CurrentTransform.GetRotation());
+						ActorDrag.CurrentTransform.SetRotation(OptQuat);
+						SetLocation(ActorDrag.Actor, ActorDrag.CurrentTransform.GetLocation(), OptQuat);
 					}
 				}
 				else if (bTranslationChanged)
@@ -849,7 +866,6 @@ void USequencerPivotTool::GizmoTransformChanged(UTransformProxy* Proxy, FTransfo
 	}
 	
 	StartDragTransform = Transform;
-
 }
 
 void USequencerPivotTool::GizmoTransformEnded(UTransformProxy* Proxy)
