Sequencer: Motion Trail: Reliance of the focused viewport would make trail ticks draw incorrectly, so got rid of viewport/viewport client calls and use FCanvas instead.
#trivial

diff --git a/Engine/Plugins/Editor/SequencerAnimTools/Source/SequencerAnimTools/Private/MotionTrailTool.cpp b/Engine/Plugins/Editor/SequencerAnimTools/Source/SequencerAnimTools/Private/MotionTrailTool.cpp
--- a/Engine/Plugins/Editor/SequencerAnimTools/Source/SequencerAnimTools/Private/MotionTrailTool.cpp
+++ b/Engine/Plugins/Editor/SequencerAnimTools/Source/SequencerAnimTools/Private/MotionTrailTool.cpp
@@ -236,21 +236,17 @@ void UMotionTrailTool::OnClicked(const FInputDeviceRay& ClickPos)
 
 void UMotionTrailTool::Render(IToolsContextRenderAPI* RenderAPI)
 {
-	FViewport* FocusedViewport = GetToolManager()->GetContextQueriesAPI()->GetFocusedViewport();
-
 	for (TUniquePtr<UE::SequencerAnimTools::FTrailHierarchy>& TrailHierarchy : TrailHierarchies)
 	{
-		TrailHierarchy->GetRenderer()->Render(RenderAPI->GetSceneView(), FocusedViewport, RenderAPI->GetPrimitiveDrawInterface());
+		TrailHierarchy->GetRenderer()->Render(RenderAPI->GetSceneView(), RenderAPI->GetPrimitiveDrawInterface());
 	}
 }
 
 void UMotionTrailTool::DrawHUD(FCanvas* Canvas, IToolsContextRenderAPI* RenderAPI)
 {
-	FViewport* FocusedViewport = GetToolManager()->GetContextQueriesAPI()->GetFocusedViewport();
-	FEditorViewportClient* ViewportClient = ModeManager->GetFocusedViewportClient();
 	for (TUniquePtr<UE::SequencerAnimTools::FTrailHierarchy>& TrailHierarchy : TrailHierarchies)
 	{
-		TrailHierarchy->GetRenderer()->DrawHUD(ViewportClient, FocusedViewport, RenderAPI->GetSceneView(), Canvas);
+		TrailHierarchy->GetRenderer()->DrawHUD(RenderAPI->GetSceneView(), Canvas);
 	}
 }
 
diff --git a/Engine/Plugins/Editor/SequencerAnimTools/Source/SequencerAnimTools/Private/Sequencer/MotionTrailMovieSceneKey.cpp b/Engine/Plugins/Editor/SequencerAnimTools/Source/SequencerAnimTools/Private/Sequencer/MotionTrailMovieSceneKey.cpp
--- a/Engine/Plugins/Editor/SequencerAnimTools/Source/SequencerAnimTools/Private/Sequencer/MotionTrailMovieSceneKey.cpp
+++ b/Engine/Plugins/Editor/SequencerAnimTools/Source/SequencerAnimTools/Private/Sequencer/MotionTrailMovieSceneKey.cpp
@@ -13,6 +13,7 @@
 #include "Tools/MotionTrailOptions.h"
 #include "CanvasItem.h"
 #include "ISequencer.h"
+#include "CanvasTypes.h"
 
 namespace UE
 {
@@ -36,7 +37,7 @@ struct HMotionTrailMovieSceneKeyProxy : public HBaseTrailProxy
 };
 IMPLEMENT_HIT_PROXY(HMotionTrailMovieSceneKeyProxy, HBaseTrailProxy);
 
-void FMotionTraiMovieScenelKeyTool::Render(const FGuid& Guid, const FSceneView* View, FViewport* Viewport, FPrimitiveDrawInterface* PDI)
+void FMotionTraiMovieScenelKeyTool::Render(const FGuid& Guid, const FSceneView* View, FPrimitiveDrawInterface* PDI)
 {
 	const bool bIsVisible = (UMotionTrailToolOptions::GetTrailOptions()->bShowKeys);
 
@@ -68,7 +69,7 @@ void FMotionTraiMovieScenelKeyTool::Render(const FGuid& Guid, const FSceneView*
 	}
 }
 
-void FMotionTraiMovieScenelKeyTool::DrawHUD(FEditorViewportClient* ViewportClient, FViewport* Viewport, const FSceneView* View, FCanvas* Canvas)
+void FMotionTraiMovieScenelKeyTool::DrawHUD(const FSceneView* View, FCanvas* Canvas)
 {
 	const bool bIsVisible = (UMotionTrailToolOptions::GetTrailOptions()->bShowKeys && UMotionTrailToolOptions::GetTrailOptions()->bShowFrameNumber);
 	if (!bIsVisible)
@@ -90,7 +91,7 @@ void FMotionTraiMovieScenelKeyTool::DrawHUD(FEditorViewportClient* ViewportClien
 			FVector2D PixelLocation;
 			if (View->WorldToPixel(FrameKeyPair.Value->Transform.GetLocation(), PixelLocation))
 			{
-				PixelLocation /= ViewportClient->GetDPIScale();
+				PixelLocation /= Canvas->GetDPIScale();
 
 				int32 FrameNumber = FFrameRate::TransformTime(FFrameTime(FrameKeyPair.Value->FrameNumber), TickResolution, DisplayRate).RoundToFrame().Value;
 				FCanvasTextItem TextItem(PixelLocation, FText::FromString(FString::FromInt(FrameNumber)), GEngine->GetMediumFont(), KeyColor);
diff --git a/Engine/Plugins/Editor/SequencerAnimTools/Source/SequencerAnimTools/Private/Sequencer/MotionTrailMovieSceneKey.h b/Engine/Plugins/Editor/SequencerAnimTools/Source/SequencerAnimTools/Private/Sequencer/MotionTrailMovieSceneKey.h
--- a/Engine/Plugins/Editor/SequencerAnimTools/Source/SequencerAnimTools/Private/Sequencer/MotionTrailMovieSceneKey.h
+++ b/Engine/Plugins/Editor/SequencerAnimTools/Source/SequencerAnimTools/Private/Sequencer/MotionTrailMovieSceneKey.h
@@ -213,8 +213,8 @@ public:
 	{}
 	~FMotionTraiMovieScenelKeyTool() {};
 
-	void Render(const FGuid& Guid, const FSceneView* View, FViewport* Viewport, FPrimitiveDrawInterface* PDI);
-	void DrawHUD(FEditorViewportClient* ViewportClient, FViewport* Viewport, const FSceneView* View, FCanvas* Canvas);
+	void Render(const FGuid& Guid, const FSceneView* View, FPrimitiveDrawInterface* PDI);
+	void DrawHUD(const FSceneView* View, FCanvas* Canvas);
 
 	bool HandleClick(const FGuid& Guid, FEditorViewportClient* InViewportClient, HHitProxy* HitProxy, FInputClick Click);
 	void StartDragging() {};
diff --git a/Engine/Plugins/Editor/SequencerAnimTools/Source/SequencerAnimTools/Private/Sequencer/MovieSceneTransformTrail.cpp b/Engine/Plugins/Editor/SequencerAnimTools/Source/SequencerAnimTools/Private/Sequencer/MovieSceneTransformTrail.cpp
--- a/Engine/Plugins/Editor/SequencerAnimTools/Source/SequencerAnimTools/Private/Sequencer/MovieSceneTransformTrail.cpp
+++ b/Engine/Plugins/Editor/SequencerAnimTools/Source/SequencerAnimTools/Private/Sequencer/MovieSceneTransformTrail.cpp
@@ -159,14 +159,14 @@ ETrailCacheState FMovieSceneTransformTrail::UpdateTrail(const FSceneContext& InS
 	return CacheState;
 }
 
-void FMovieSceneTransformTrail::DrawHUD(FEditorViewportClient* ViewportClient, FViewport* Viewport, const FSceneView* View, FCanvas* Canvas)
+void FMovieSceneTransformTrail::DrawHUD(const FSceneView* View, FCanvas* Canvas)
 {
-	KeyTool->DrawHUD(ViewportClient, Viewport, View, Canvas);
+	KeyTool->DrawHUD(View, Canvas);
 }
 
-void FMovieSceneTransformTrail::Render(const FGuid& Guid, const FSceneView* View, FViewport* Viewport, FPrimitiveDrawInterface* PDI)
+void FMovieSceneTransformTrail::Render(const FGuid& Guid, const FSceneView* View, FPrimitiveDrawInterface* PDI)
 {
-	KeyTool->Render(Guid, View, Viewport, PDI);
+	KeyTool->Render(Guid, View, PDI);
 }
 
 bool FMovieSceneTransformTrail::HandleAltClick(FEditorViewportClient* InViewportClient, HMotionTrailProxy* Proxy, FInputClick Click)
diff --git a/Engine/Plugins/Editor/SequencerAnimTools/Source/SequencerAnimTools/Private/Sequencer/MovieSceneTransformTrail.h b/Engine/Plugins/Editor/SequencerAnimTools/Source/SequencerAnimTools/Private/Sequencer/MovieSceneTransformTrail.h
--- a/Engine/Plugins/Editor/SequencerAnimTools/Source/SequencerAnimTools/Private/Sequencer/MovieSceneTransformTrail.h
+++ b/Engine/Plugins/Editor/SequencerAnimTools/Source/SequencerAnimTools/Private/Sequencer/MovieSceneTransformTrail.h
@@ -27,8 +27,8 @@ public:
 	virtual ETrailCacheState UpdateTrail(const FSceneContext& InSceneContext) override;
 	virtual FTrajectoryCache* GetTrajectoryTransforms() override { return TrajectoryCache.Get(); }
 	virtual TRange<double> GetEffectiveRange() const override { return CachedEffectiveRange; }
-	virtual void Render(const FGuid& Guid, const FSceneView* View, FViewport* Viewport, FPrimitiveDrawInterface* PDI) override;
-	virtual void DrawHUD(FEditorViewportClient* ViewportClient, FViewport* Viewport, const FSceneView* View, FCanvas* Canvas) override;
+	virtual void Render(const FGuid& Guid, const FSceneView* View,  FPrimitiveDrawInterface* PDI) override;
+	virtual void DrawHUD(const FSceneView* View, FCanvas* Canvas) override;
 
 	virtual bool HandleClick(const FGuid& Guid, FEditorViewportClient* InViewportClient, HHitProxy* HitProxy, FInputClick Click) override;
 	virtual bool IsAnythingSelected() const override;
diff --git a/Engine/Plugins/Editor/SequencerAnimTools/Source/SequencerAnimTools/Private/TrailHierarchy.cpp b/Engine/Plugins/Editor/SequencerAnimTools/Source/SequencerAnimTools/Private/TrailHierarchy.cpp
--- a/Engine/Plugins/Editor/SequencerAnimTools/Source/SequencerAnimTools/Private/TrailHierarchy.cpp
+++ b/Engine/Plugins/Editor/SequencerAnimTools/Source/SequencerAnimTools/Private/TrailHierarchy.cpp
@@ -11,6 +11,7 @@
 #include "EditorModeManager.h"
 #include "Tools/MotionTrailOptions.h"
 #include "Widgets/Input/SButton.h"
+#include "CanvasTypes.h"
 
 #define LOCTEXT_NAMESPACE "MotionTrailEditorMode"
 
@@ -22,7 +23,7 @@ namespace SequencerAnimTools
 IMPLEMENT_HIT_PROXY(HBaseTrailProxy, HHitProxy);
 IMPLEMENT_HIT_PROXY(HMotionTrailProxy, HBaseTrailProxy);
 
-void FTrailHierarchyRenderer::Render(const FSceneView* View, FViewport* Viewport, FPrimitiveDrawInterface* PDI)
+void FTrailHierarchyRenderer::Render(const FSceneView* View, FPrimitiveDrawInterface* PDI)
 {
 	const FDateTime RenderStartTime = FDateTime::Now();
 	
@@ -39,7 +40,7 @@ void FTrailHierarchyRenderer::Render(const FSceneView* View, FViewport* Viewport
 		{
 			FDisplayContext DisplayContext = {
 				GuidTrailPair.Key,
-				FTrailScreenSpaceTransform(View, Viewport),
+				FTrailScreenSpaceTransform(View),
 				UMotionTrailToolOptions::GetTrailOptions()->SecondsPerMark,
 				OwningHierarchy->GetViewRange(),
 				OwningHierarchy->GetSecondsPerSegment()
@@ -71,7 +72,7 @@ void FTrailHierarchyRenderer::Render(const FSceneView* View, FViewport* Viewport
 					}
 				}
 			}
-			GuidTrailPair.Value->Render(GuidTrailPair.Key,View,Viewport,PDI);
+			GuidTrailPair.Value->Render(GuidTrailPair.Key,View,PDI);
 
 		}
 		else
@@ -88,14 +89,14 @@ void FTrailHierarchyRenderer::Render(const FSceneView* View, FViewport* Viewport
 	OwningHierarchy->GetTimingStats().Add("FTrailHierarchyRenderer::Render", RenderTimespan);
 }
 
-void FTrailHierarchyRenderer::DrawHUD(FEditorViewportClient* ViewportClient, FViewport* Viewport, const FSceneView* View, FCanvas* Canvas)
+void FTrailHierarchyRenderer::DrawHUD(const FSceneView* View, FCanvas* Canvas)
 {
 
 	if (UMotionTrailToolOptions::GetTrailOptions()->bShowMarks == false)
 	{
 		for (const TPair<FGuid, TUniquePtr<FTrail>>& GuidTrailPair : OwningHierarchy->GetAllTrails())
 		{
-			GuidTrailPair.Value->DrawHUD( ViewportClient, Viewport, View, Canvas);
+			GuidTrailPair.Value->DrawHUD( View, Canvas);
 		}
 		return;
 	}	
@@ -113,7 +114,7 @@ void FTrailHierarchyRenderer::DrawHUD(FEditorViewportClient* ViewportClient, FVi
 		{
 			FDisplayContext DisplayContext = {
 				GuidTrailPair.Key,
-				FTrailScreenSpaceTransform(View, Viewport, ViewportClient->GetDPIScale()),
+				FTrailScreenSpaceTransform(View, Canvas->GetDPIScale()),
 				SecondsPerMark,
 				OwningHierarchy->GetViewRange(),
 				OwningHierarchy->GetSecondsPerSegment()
@@ -132,7 +133,7 @@ void FTrailHierarchyRenderer::DrawHUD(FEditorViewportClient* ViewportClient, FVi
 				Canvas->DrawItem(LineItem);
 			}
 		}
-		GuidTrailPair.Value->DrawHUD(ViewportClient, Viewport, View, Canvas);
+		GuidTrailPair.Value->DrawHUD(View, Canvas);
 
 	}
 
diff --git a/Engine/Plugins/Editor/SequencerAnimTools/Source/SequencerAnimTools/Private/TrajectoryDrawInfo.cpp b/Engine/Plugins/Editor/SequencerAnimTools/Source/SequencerAnimTools/Private/TrajectoryDrawInfo.cpp
--- a/Engine/Plugins/Editor/SequencerAnimTools/Source/SequencerAnimTools/Private/TrajectoryDrawInfo.cpp
+++ b/Engine/Plugins/Editor/SequencerAnimTools/Source/SequencerAnimTools/Private/TrajectoryDrawInfo.cpp
@@ -11,12 +11,12 @@ namespace SequencerAnimTools
 
 TOptional<FVector2D> FTrailScreenSpaceTransform::ProjectPoint(const FVector& Point) const
 {
-	const FPlane Projection = View->Project(Point);
-	if (Projection.W > 0.0f)
+
+	FVector2D PixelLocation;
+	if (View->WorldToPixel(Point, PixelLocation))
 	{
-		const float XPos = HalfScreenSize.X + (HalfScreenSize.X * Projection.X * 1.0f);
-		const float YPos = HalfScreenSize.Y + (HalfScreenSize.Y * Projection.Y * -1.0f);
-		return FVector2D(XPos, YPos);
+		PixelLocation /= DPIScale;
+		return PixelLocation;
 	}
 
 	return TOptional<FVector2D>();
diff --git a/Engine/Plugins/Editor/SequencerAnimTools/Source/SequencerAnimTools/Public/Trail.h b/Engine/Plugins/Editor/SequencerAnimTools/Source/SequencerAnimTools/Public/Trail.h
--- a/Engine/Plugins/Editor/SequencerAnimTools/Source/SequencerAnimTools/Public/Trail.h
+++ b/Engine/Plugins/Editor/SequencerAnimTools/Source/SequencerAnimTools/Public/Trail.h
@@ -95,8 +95,8 @@ public:
 	virtual FTrajectoryCache* GetTrajectoryTransforms() = 0;
 
 	//Additional Render/HitTest event handling for specific trails, usually let default renderer handle it
-	virtual void Render(const FGuid& Guid, const FSceneView* View, FViewport* Viewport, FPrimitiveDrawInterface* PDI) {};
-	virtual void DrawHUD(FEditorViewportClient* ViewportClient, FViewport* Viewport, const FSceneView* View, FCanvas* Canvas) {};
+	virtual void Render(const FGuid& Guid, const FSceneView* View, FPrimitiveDrawInterface* PDI) {};
+	virtual void DrawHUD(const FSceneView* View, FCanvas* Canvas) {};
 	virtual bool HandleClick(const FGuid& Guid, FEditorViewportClient* InViewportClient, HHitProxy* HitProxy, FInputClick Click) { return false; }
 	virtual bool IsAnythingSelected() const { return false; }
 	virtual bool IsAnythingSelected(FVector& OutVectorPosition)const { return false;}
diff --git a/Engine/Plugins/Editor/SequencerAnimTools/Source/SequencerAnimTools/Public/TrailHierarchy.h b/Engine/Plugins/Editor/SequencerAnimTools/Source/SequencerAnimTools/Public/TrailHierarchy.h
--- a/Engine/Plugins/Editor/SequencerAnimTools/Source/SequencerAnimTools/Public/TrailHierarchy.h
+++ b/Engine/Plugins/Editor/SequencerAnimTools/Source/SequencerAnimTools/Public/TrailHierarchy.h
@@ -52,8 +52,8 @@ class ITrailHierarchyRenderer
 {
 public:
 	virtual ~ITrailHierarchyRenderer() {}
-	virtual void Render(const FSceneView* View, FViewport* Viewport, FPrimitiveDrawInterface* PDI) = 0;
-	virtual void DrawHUD(FEditorViewportClient* ViewportClient, FViewport* Viewport, const FSceneView* View, FCanvas* Canvas) = 0;
+	virtual void Render(const FSceneView* View, FPrimitiveDrawInterface* PDI) = 0;
+	virtual void DrawHUD(const FSceneView* View, FCanvas* Canvas) = 0;
 };
 
 class FTrailHierarchyRenderer : public ITrailHierarchyRenderer
@@ -63,8 +63,8 @@ public:
 		: OwningHierarchy(InOwningHierarchy)
 	{}
 
-	virtual void Render(const FSceneView* View, FViewport* Viewport, FPrimitiveDrawInterface* PDI) override;
-	virtual void DrawHUD(FEditorViewportClient* ViewportClient, FViewport* Viewport, const FSceneView* View, FCanvas* Canvas) override;
+	virtual void Render(const FSceneView* View, FPrimitiveDrawInterface* PDI) override;
+	virtual void DrawHUD(const FSceneView* View, FCanvas* Canvas) override;
 
 private:
 	class FTrailHierarchy* OwningHierarchy;
diff --git a/Engine/Plugins/Editor/SequencerAnimTools/Source/SequencerAnimTools/Public/TrajectoryDrawInfo.h b/Engine/Plugins/Editor/SequencerAnimTools/Source/SequencerAnimTools/Public/TrajectoryDrawInfo.h
--- a/Engine/Plugins/Editor/SequencerAnimTools/Source/SequencerAnimTools/Public/TrajectoryDrawInfo.h
+++ b/Engine/Plugins/Editor/SequencerAnimTools/Source/SequencerAnimTools/Public/TrajectoryDrawInfo.h
@@ -25,16 +25,16 @@ class FTrajectoryCache;
 class FTrailScreenSpaceTransform
 {
 public:
-	FTrailScreenSpaceTransform(const FSceneView* InView, FViewport* InViewport, const float InDPIScale = 1.0f)
+	FTrailScreenSpaceTransform(const FSceneView* InView,  const float InDPIScale = 1.0f)
 		: View(InView)
-		, HalfScreenSize((InViewport->GetSizeXY().X / InDPIScale) * 0.5, (InViewport->GetSizeXY().Y / InDPIScale) * 0.5)
+		, DPIScale(InDPIScale)
 	{}
 
 	TOptional<FVector2D> ProjectPoint(const FVector& Point) const;
 
 private:
 	const FSceneView* View;
-	FVector2D HalfScreenSize;
+	const float DPIScale;
 };
 
 struct FDisplayContext
