SequencerAnimPivotTool to deploy interactionscope correctly

diff --git a/Engine/Plugins/Editor/SequencerAnimTools/Source/SequencerAnimTools/Private/SequencerAnimEditPivotTool.cpp b/Engine/Plugins/Editor/SequencerAnimTools/Source/SequencerAnimTools/Private/SequencerAnimEditPivotTool.cpp
--- a/Engine/Plugins/Editor/SequencerAnimTools/Source/SequencerAnimTools/Private/SequencerAnimEditPivotTool.cpp
+++ b/Engine/Plugins/Editor/SequencerAnimTools/Source/SequencerAnimTools/Private/SequencerAnimEditPivotTool.cpp
@@ -701,6 +701,33 @@ void USequencerPivotTool::GizmoTransformStarted(UTransformProxy* Proxy)
 	bGizmoBeingDragged = true;
 	bManipulatorMadeChange = false;
 	GizmoTransform = StartDragTransform;
+
+	InteractionScopes.Reset();
+	if(bInPivotMode)
+	{
+		TMap<UControlRig*, int32> RigToScopeIndex;
+		for(int32 IndexA = 0; IndexA < ControlRigDrags.Num(); IndexA++)
+		{
+			const FControlRigSelectionDuringDrag& ControlRigDrag = ControlRigDrags[IndexA];
+
+			// if we are hitting this for the first time
+			const int32 ScopeIndex = RigToScopeIndex.FindOrAdd(ControlRigDrag.ControlRig, InteractionScopes.Num());
+			if(!InteractionScopes.IsValidIndex(ScopeIndex))
+			{
+				// get all keys on the same control rig
+				TArray<FRigElementKey> Keys = {FRigElementKey(ControlRigDrag.ControlName, ERigElementType::Control)};
+				for(int32 IndexB = IndexA + 1; IndexB < ControlRigDrags.Num(); IndexB++)
+				{
+					if(ControlRigDrags[IndexB].ControlRig != ControlRigDrag.ControlRig)
+					{
+						continue;
+					}
+					Keys.Add(FRigElementKey(ControlRigDrags[IndexB].ControlName, ERigElementType::Control));
+				}
+				InteractionScopes.Emplace(MakeShareable(new FControlRigInteractionScope(ControlRigDrag.ControlRig, Keys)));
+			}
+		}
+	}
 }
 
 
@@ -881,6 +908,7 @@ void USequencerPivotTool::GizmoTransformEnded(UTransformProxy* Proxy)
 	}
 	bGizmoBeingDragged = false;
 	bManipulatorMadeChange = false;
+	InteractionScopes.Reset();
 	UpdateGizmoTransform();
 	SavePivotTransforms();
 }
diff --git a/Engine/Plugins/Editor/SequencerAnimTools/Source/SequencerAnimTools/Public/SequencerAnimEditPivotTool.h b/Engine/Plugins/Editor/SequencerAnimTools/Source/SequencerAnimTools/Public/SequencerAnimEditPivotTool.h
--- a/Engine/Plugins/Editor/SequencerAnimTools/Source/SequencerAnimTools/Public/SequencerAnimEditPivotTool.h
+++ b/Engine/Plugins/Editor/SequencerAnimTools/Source/SequencerAnimTools/Public/SequencerAnimEditPivotTool.h
@@ -13,6 +13,7 @@
 #include "Framework/Commands/UICommandInfo.h"
 #include "BaseGizmos/TransformProxy.h"
 #include "BaseSequencerAnimTool.h"
+#include "ControlRig.h"
 #include "Styling/AppStyle.h"
 #include "SequencerAnimEditPivotTool.generated.h"
 
@@ -20,7 +21,6 @@ class USingleClickInputBehavior;
 class UClickDragInputBehavior;
 class UCombinedTransformGizmo;
 class ULevelSequence;
-class UControlRig;
 struct FRigControlElement;
 class ISequencer;
 class FUICommandList;
@@ -189,6 +189,7 @@ protected:
 	int32 TransactionIndex = -1;
 	TArray<FControlRigSelectionDuringDrag> ControlRigDrags;
 	TArray<FActorSelectonDuringDrag> ActorDrags;
+	TArray<TSharedPtr<FControlRigInteractionScope>> InteractionScopes;
 
 	//since we are selection based we can cache this
 	ULevelSequence* LevelSequence;
