Sequencer: Anim Tools: For Overlays we need to check all the viewports to remove it, getting active is not always accurate unforunately.

diff --git a/Engine/Plugins/Animation/ControlRig/Source/ControlRigEditor/Private/EditMode/ControlRigEditModeToolkit.cpp b/Engine/Plugins/Animation/ControlRig/Source/ControlRigEditor/Private/EditMode/ControlRigEditModeToolkit.cpp
--- a/Engine/Plugins/Animation/ControlRig/Source/ControlRigEditor/Private/EditMode/ControlRigEditModeToolkit.cpp
+++ b/Engine/Plugins/Animation/ControlRig/Source/ControlRigEditor/Private/EditMode/ControlRigEditModeToolkit.cpp
@@ -23,6 +23,9 @@
 #include "EditMode/SControlRigDetails.h"
 #include "EditMode/SControlRigOutliner.h"
 #include "EditMode/SControlRigSpacePicker.h"
+#include "LevelEditor.h"
+#include "SLevelViewport.h"
+
 #define LOCTEXT_NAMESPACE "FControlRigEditModeToolkit"
 
 namespace 
@@ -297,7 +300,22 @@ void FControlRigEditModeToolkit::TryRemoveTweenOverlay()
 {
 	if (IsHosted() && TweenWidgetParent)
 	{
-		GetToolkitHost()->RemoveViewportOverlayWidget(TweenWidgetParent.ToSharedRef());
+		if (FLevelEditorModule* LevelEditorModule = FModuleManager::GetModulePtr<FLevelEditorModule>(TEXT("LevelEditor")))
+		{
+			if (TSharedPtr<ILevelEditor> LevelEditor = LevelEditorModule->GetFirstLevelEditor())
+			{
+				for (TSharedPtr<SLevelViewport> LevelViewport : LevelEditor->GetViewports())
+				{
+					if (LevelViewport.IsValid())
+					{
+						if (TweenWidgetParent)
+						{
+							LevelViewport->RemoveOverlayWidget(TweenWidgetParent.ToSharedRef());
+						}
+					}
+				}
+			}
+		}
 	}
 }
 
diff --git a/Engine/Plugins/Editor/SequencerAnimTools/Source/SequencerAnimTools/Private/SequencerAnimEditPivotTool.cpp b/Engine/Plugins/Editor/SequencerAnimTools/Source/SequencerAnimTools/Private/SequencerAnimEditPivotTool.cpp
--- a/Engine/Plugins/Editor/SequencerAnimTools/Source/SequencerAnimTools/Private/SequencerAnimEditPivotTool.cpp
+++ b/Engine/Plugins/Editor/SequencerAnimTools/Source/SequencerAnimTools/Private/SequencerAnimEditPivotTool.cpp
@@ -50,6 +50,7 @@
 #include "Styling/StyleColors.h"
 #include "EditorModes.h"
 #include "ControlRigObjectBinding.h"
+#include "SLevelViewport.h"
 
 #define LOCTEXT_NAMESPACE "SequencerAnimTools"
 
@@ -1229,30 +1230,48 @@ FMargin USequencerPivotTool::GetPivotOverlayPadding() const
 
 void USequencerPivotTool::RemoveAndDestroyPivotOverlay()
 {
+	if (PivotWidget.IsValid() == false)
+	{
+		return;
+	}
 	if (FLevelEditorModule* LevelEditorModule = FModuleManager::GetModulePtr<FLevelEditorModule>(TEXT("LevelEditor")))
 	{
-		TSharedPtr<IAssetViewport> ActiveLevelViewport = LevelEditorModule->GetFirstActiveViewport();
-		if (ActiveLevelViewport.IsValid())
+		if (TSharedPtr<ILevelEditor> LevelEditor = LevelEditorModule->GetFirstLevelEditor())
 		{
-			if (PivotWidget)
+			for (TSharedPtr<SLevelViewport> LevelViewport : LevelEditor->GetViewports())
 			{
-				ActiveLevelViewport->RemoveOverlayWidget(PivotWidget.ToSharedRef());
-				PivotWidget.Reset();
+				if (LevelViewport.IsValid())
+				{
+					if (PivotWidget)
+					{
+						LevelViewport->RemoveOverlayWidget(PivotWidget.ToSharedRef());
+					}
+				}
 			}
+			PivotWidget.Reset();
 		}
 	}
 }
 
 void USequencerPivotTool::TryRemovePivotOverlay()
 {
+	if (PivotWidget.IsValid() == false)
+	{
+		return;
+	}
 	if (FLevelEditorModule* LevelEditorModule = FModuleManager::GetModulePtr<FLevelEditorModule>(TEXT("LevelEditor")))
 	{
-		TSharedPtr<IAssetViewport> ActiveLevelViewport = LevelEditorModule->GetFirstActiveViewport();
-		if (ActiveLevelViewport.IsValid())
+		if (TSharedPtr<ILevelEditor> LevelEditor = LevelEditorModule->GetFirstLevelEditor())
 		{
-			if (PivotWidget)
+			for (TSharedPtr<SLevelViewport> LevelViewport : LevelEditor->GetViewports())
 			{
-				ActiveLevelViewport->RemoveOverlayWidget(PivotWidget.ToSharedRef());
+				if (LevelViewport.IsValid())
+				{
+					if (PivotWidget)
+					{
+						LevelViewport->RemoveOverlayWidget(PivotWidget.ToSharedRef());
+					}
+				}
 			}
 		}
 	}
