Sequencer: Pivot Tool : Regression, get click/shift options working again, needed to save out the transforms since we need it for each tick, not sure how this regressed.

diff --git a/Engine/Plugins/Editor/SequencerAnimTools/Source/SequencerAnimTools/Private/SequencerAnimEditPivotTool.cpp b/Engine/Plugins/Editor/SequencerAnimTools/Source/SequencerAnimTools/Private/SequencerAnimEditPivotTool.cpp
--- a/Engine/Plugins/Editor/SequencerAnimTools/Source/SequencerAnimTools/Private/SequencerAnimEditPivotTool.cpp
+++ b/Engine/Plugins/Editor/SequencerAnimTools/Source/SequencerAnimTools/Private/SequencerAnimEditPivotTool.cpp
@@ -399,6 +399,7 @@ void USequencerPivotTool::UpdateTransformAndSelectionOnEntering()
 		}
 		SetGizmoBasedOnSelection(!bShiftPressedWhenStarted && !bCtrlPressedWhenStarted);
 	}
+	SavePivotTransforms();
 }
 
 bool USequencerPivotTool::ProcessCommandBindings(const FKey Key, const bool bRepeat) const
@@ -964,6 +965,7 @@ void USequencerPivotTool::OnClicked(const FInputDeviceRay& ClickPos)
 		{
 			GizmoTransform.SetLocation(HitLocation);
 			UpdateGizmoTransform();
+			SavePivotTransforms();
 		}
 	}
 }
