Sequencer: Speculative fix for motion trail crash

diff --git a/Engine/Plugins/Editor/SequencerAnimTools/Source/SequencerAnimTools/Private/Sequencer/SequencerTrailHierarchy.cpp b/Engine/Plugins/Editor/SequencerAnimTools/Source/SequencerAnimTools/Private/Sequencer/SequencerTrailHierarchy.cpp
--- a/Engine/Plugins/Editor/SequencerAnimTools/Source/SequencerAnimTools/Private/Sequencer/SequencerTrailHierarchy.cpp
+++ b/Engine/Plugins/Editor/SequencerAnimTools/Source/SequencerAnimTools/Private/Sequencer/SequencerTrailHierarchy.cpp
@@ -276,7 +276,12 @@ void FSequencerTrailHierarchy::UpdateControlRig(const TArray<FFrameNumber> &Fram
 				}
 				for (FTrailControlTransforms& TrailControlTransform : TrailControlTransforms)
 				{
-					if (bUseEditedTimes == false)
+					if (!TrailControlTransform.Trail)
+					{
+						continue;
+					}
+
+					if (bUseEditedTimes == false && TrailControlTransform.Trail->GetTrajectoryTransforms())
 					{
 						TrailControlTransform.Trail->GetTrajectoryTransforms()->SetTransforms(TrailControlTransform.Transforms, ControlRigParentWorldTransforms);
 					}
