Control Rig: Fix crash when enabling motion trails
#rnx

diff --git a/Engine/Plugins/Editor/SequencerAnimTools/Source/SequencerAnimTools/Private/Sequencer/SequencerTrailHierarchy.cpp b/Engine/Plugins/Editor/SequencerAnimTools/Source/SequencerAnimTools/Private/Sequencer/SequencerTrailHierarchy.cpp
--- a/Engine/Plugins/Editor/SequencerAnimTools/Source/SequencerAnimTools/Private/Sequencer/SequencerTrailHierarchy.cpp
+++ b/Engine/Plugins/Editor/SequencerAnimTools/Source/SequencerAnimTools/Private/Sequencer/SequencerTrailHierarchy.cpp
@@ -212,7 +212,7 @@ void FSequencerTrailHierarchy::RemoveTrail(const FGuid& Key)
 struct FTrailControlTransforms
 {
 	FName ControlName;
-	FTrail* Trail;
+	FGuid ElementGuid;
 	TArray<FTransform> Transforms;
 };
 
@@ -248,7 +248,7 @@ void FSequencerTrailHierarchy::UpdateControlRig(const TArray<FFrameNumber> &Fram
 				int32 PairIndex = 0;
 				for (TPair<FName, FGuid >& Pair : CompMapPair)
 				{
-					TrailControlTransforms[PairIndex].Trail = AllTrails[Pair.Value].Get();
+					TrailControlTransforms[PairIndex].ElementGuid = Pair.Value;
 					TrailControlTransforms[PairIndex].ControlName = Pair.Key;
 					TrailControlTransforms[PairIndex].Transforms.SetNum(Frames.Num());
 					++PairIndex;
@@ -270,22 +270,23 @@ void FSequencerTrailHierarchy::UpdateControlRig(const TArray<FFrameNumber> &Fram
 						if (bUseEditedTimes)
 						{
 							double Sec = TickResolution.AsSeconds(FFrameTime(FrameNumber));
-							TrailControlTransform.Trail->GetTrajectoryTransforms()->Set(Sec, TrailControlTransform.Transforms[Index]);
+							AllTrails[TrailControlTransform.ElementGuid]->GetTrajectoryTransforms()->Set(Sec, TrailControlTransform.Transforms[Index]);
 						}
 					}
 				}
 				for (FTrailControlTransforms& TrailControlTransform : TrailControlTransforms)
 				{
-					if (!TrailControlTransform.Trail)
+					if (!AllTrails.Contains(TrailControlTransform.ElementGuid))
 					{
 						continue;
 					}
 
-					if (bUseEditedTimes == false && TrailControlTransform.Trail->GetTrajectoryTransforms())
+					FTrail* Trail = AllTrails[TrailControlTransform.ElementGuid].Get();
+					if (bUseEditedTimes == false && Trail->GetTrajectoryTransforms())
 					{
-						TrailControlTransform.Trail->GetTrajectoryTransforms()->SetTransforms(TrailControlTransform.Transforms, ControlRigParentWorldTransforms);
+						Trail->GetTrajectoryTransforms()->SetTransforms(TrailControlTransform.Transforms, ControlRigParentWorldTransforms);
 					}
-					TrailControlTransform.Trail->UpdateKeysInRange(GetViewRange());
+					Trail->UpdateKeysInRange(GetViewRange());
 				}
 			}
 		}
