Sequencer: MotionTrails speculative fix for crash, no longer let user specify how many ticks to do.

diff --git a/Engine/Plugins/Editor/SequencerAnimTools/Source/SequencerAnimTools/Private/Sequencer/SequencerTrailHierarchy.cpp b/Engine/Plugins/Editor/SequencerAnimTools/Source/SequencerAnimTools/Private/Sequencer/SequencerTrailHierarchy.cpp
--- a/Engine/Plugins/Editor/SequencerAnimTools/Source/SequencerAnimTools/Private/Sequencer/SequencerTrailHierarchy.cpp
+++ b/Engine/Plugins/Editor/SequencerAnimTools/Source/SequencerAnimTools/Private/Sequencer/SequencerTrailHierarchy.cpp
@@ -162,15 +162,13 @@ void FSequencerTrailHierarchy::OnActorsChangedSomehow(TArray<AActor*>& InActors)
 double FSequencerTrailHierarchy::GetSecondsPerSegment() const 
 { 
 	double SecondsPerFrame =  WeakSequencer.Pin()->GetFocusedDisplayRate().AsInterval();
-	//evalsperframe will not be less than one
-	return SecondsPerFrame / double(UMotionTrailToolOptions::GetTrailOptions()->EvalsPerFrame);
+	return SecondsPerFrame; // no longer use this to avoid divide by zero options / double(UMotionTrailToolOptions::GetTrailOptions()->EvalsPerFrame);
 }
 
 FFrameNumber FSequencerTrailHierarchy::GetFramesPerSegment() const
 {
 	FFrameNumber FramesPerTick = FFrameRate::TransformTime(FFrameNumber(1), WeakSequencer.Pin()->GetFocusedDisplayRate(), 
 		WeakSequencer.Pin()->GetFocusedTickResolution()).RoundToFrame();
-	FramesPerTick.Value /= (UMotionTrailToolOptions::GetTrailOptions()->EvalsPerFrame);
 	return FramesPerTick;
 }
 
diff --git a/Engine/Source/Editor/Sequencer/Public/Tools/MotionTrailOptions.h b/Engine/Source/Editor/Sequencer/Public/Tools/MotionTrailOptions.h
--- a/Engine/Source/Editor/Sequencer/Public/Tools/MotionTrailOptions.h
+++ b/Engine/Source/Editor/Sequencer/Public/Tools/MotionTrailOptions.h
@@ -59,8 +59,7 @@ public:
 	UPROPERTY(EditAnywhere, Category = Trail, Meta = (EditCondition = "!bShowFullTrail", ClampMin = "0"))
 	int32 FramesAfter;
 
-	/** The number of evaluations per frame */
-	UPROPERTY(EditAnywhere, Category = Trail, Meta = (ClampMin = "1.0"))
+	/** No longer exposed and clamped to 1 The number of evaluations per frame */
 	int32 EvalsPerFrame;
 	
 	/** Whether or not to show keys on the motion trail */
