Sequencer: Motion Trail: Handle invalid intersections

diff --git a/Engine/Plugins/Editor/SequencerAnimTools/Source/SequencerAnimTools/Private/TrajectoryCache.cpp b/Engine/Plugins/Editor/SequencerAnimTools/Source/SequencerAnimTools/Private/TrajectoryCache.cpp
--- a/Engine/Plugins/Editor/SequencerAnimTools/Source/SequencerAnimTools/Private/TrajectoryCache.cpp
+++ b/Engine/Plugins/Editor/SequencerAnimTools/Source/SequencerAnimTools/Private/TrajectoryCache.cpp
@@ -103,17 +103,19 @@ void FArrayTrajectoryCache::GetInterp(const double InTime,FTransform& Transform,
 TArray<double> FArrayTrajectoryCache::GetAllTimesInRange(const TRange<double>& InRange) const
 {
 	TRange<double> GenRange = TRange<double>::Intersection({ TrackRange, InRange });
-
 	TArray<double> AllTimesInRange;
-	AllTimesInRange.Reserve(int(GenRange.Size<double>() / Spacing) + 1);
-	const double FirstTick = FMath::FloorToDouble((GenRange.GetLowerBoundValue() / Spacing) + KINDA_SMALL_NUMBER) * Spacing;
-	for (double TickItr = FirstTick + KINDA_SMALL_NUMBER; TickItr < GenRange.GetUpperBoundValue(); TickItr += Spacing)
+	//intersection can return lower > upper when don't intersect
+	if (GenRange.IsEmpty() == false && GenRange.GetLowerBoundValue() < GenRange.GetUpperBoundValue())
 	{
-		AllTimesInRange.Add(TickItr);
-	}
-
-	AllTimesInRange.Add(GenRange.GetUpperBoundValue());
+		AllTimesInRange.Reserve(int(GenRange.Size<double>() / Spacing) + 1);
+		const double FirstTick = FMath::FloorToDouble((GenRange.GetLowerBoundValue() / Spacing) + KINDA_SMALL_NUMBER) * Spacing;
+		for (double TickItr = FirstTick + KINDA_SMALL_NUMBER; TickItr < GenRange.GetUpperBoundValue(); TickItr += Spacing)
+		{
+			AllTimesInRange.Add(TickItr);
+		}
 
+		AllTimesInRange.Add(GenRange.GetUpperBoundValue());
+	}
 	return AllTimesInRange;
 }
 
