[Truncation Warnings] Update IToolkitHost interface to return a slate / float vector 2 for viewport size

#rnx

diff --git a/Engine/Plugins/Animation/ControlRig/Source/ControlRigEditor/Private/EditMode/ControlRigEditModeToolkit.cpp b/Engine/Plugins/Animation/ControlRig/Source/ControlRigEditor/Private/EditMode/ControlRigEditModeToolkit.cpp
--- a/Engine/Plugins/Animation/ControlRig/Source/ControlRigEditor/Private/EditMode/ControlRigEditModeToolkit.cpp
+++ b/Engine/Plugins/Animation/ControlRig/Source/ControlRigEditor/Private/EditMode/ControlRigEditModeToolkit.cpp
@@ -245,7 +245,7 @@ void FControlRigEditModeToolkit::CreateAndShowTweenOverlay()
 
 	if (NewTweenWidgetLocation.IsZero())
 	{
-		const FVector2D ActiveViewportSize = GetToolkitHost()->GetActiveViewportSize();
+		const FVector2f ActiveViewportSize = GetToolkitHost()->GetActiveViewportWidgetSize();
 		NewTweenWidgetLocation.X = ActiveViewportSize.X / 2.0f;
 		NewTweenWidgetLocation.Y = FMath::Max(ActiveViewportSize.Y - 100.0f, 0);
 	}
@@ -345,7 +345,7 @@ void FControlRigEditModeToolkit::TryRemoveTweenOverlay()
 
 void FControlRigEditModeToolkit::UpdateTweenWidgetLocation(const FVector2D InLocation)
 {
-	const FVector2D ActiveViewportSize = GetToolkitHost()->GetActiveViewportSize();
+	const FVector2f ActiveViewportSize = GetToolkitHost()->GetActiveViewportWidgetSize();
 	FVector2D ScreenPos = InLocation;
 
 	const float EdgeFactor = 0.97f;
@@ -358,7 +358,7 @@ void FControlRigEditModeToolkit::UpdateTweenWidgetLocation(const FVector2D InLoc
 	{
 		// reset the location if it was placed out of bounds
 		ScreenPos.X = ActiveViewportSize.X / 2.0f;
-		ScreenPos.Y = FMath::Max(ActiveViewportSize.Y - 100.0f, 0);
+		ScreenPos.Y = FMath::Max(ActiveViewportSize.Y - 100.0f, 0.f);
 	}
 	InViewportTweenWidgetLocation = ScreenPos;
 	UControlRigEditModeSettings* ControlRigEditModeSettings = GetMutableDefault<UControlRigEditModeSettings>();
diff --git a/Engine/Plugins/Editor/SequencerAnimTools/Source/SequencerAnimTools/Private/SequencerAnimEditPivotTool.cpp b/Engine/Plugins/Editor/SequencerAnimTools/Source/SequencerAnimTools/Private/SequencerAnimEditPivotTool.cpp
--- a/Engine/Plugins/Editor/SequencerAnimTools/Source/SequencerAnimTools/Private/SequencerAnimEditPivotTool.cpp
+++ b/Engine/Plugins/Editor/SequencerAnimTools/Source/SequencerAnimTools/Private/SequencerAnimEditPivotTool.cpp
@@ -1027,7 +1027,7 @@ void USequencerPivotTool::Render(IToolsContextRenderAPI* RenderAPI)
 
 //Pivot Overlay functions
 
-FVector2D USequencerPivotTool::LastPivotOverlayLocation = FVector2D(0.0, 0.0);
+FVector2f USequencerPivotTool::LastPivotOverlayLocation = FVector2f::ZeroVector;
 
 class SPivotOverlayWidget : public SCompoundWidget
 {
@@ -1159,7 +1159,7 @@ private:
 	}
 };
 
-static FVector2D GetActiveViewportSize(TSharedPtr<IAssetViewport>& ActiveViewport)
+static FVector2f GetActiveViewportSize(TSharedPtr<IAssetViewport>& ActiveViewport)
 {
 	FSceneViewFamilyContext ViewFamily(FSceneViewFamily::ConstructionValues(
 		ActiveViewport->GetActiveViewport(),
@@ -1171,7 +1171,7 @@ static FVector2D GetActiveViewportSize(TSharedPtr<IAssetViewport>& ActiveViewpor
 	const float InvDpiScale = 1.0f / ActiveViewport->GetAssetViewportClient().GetDPIScale();
 	const float MaxX = SceneView->UnscaledViewRect.Width() * InvDpiScale;
 	const float MaxY = SceneView->UnscaledViewRect.Height() * InvDpiScale;
-	return FVector2D(MaxX, MaxY);
+	return FVector2f(MaxX, MaxY);
 }
 
 void USequencerPivotTool::CreateAndShowPivotOverlay()
@@ -1181,16 +1181,16 @@ void USequencerPivotTool::CreateAndShowPivotOverlay()
 		TSharedPtr<IAssetViewport> ActiveLevelViewport = LevelEditorModule->GetFirstActiveViewport();
 		if (ActiveLevelViewport.IsValid())
 		{
-			FVector2D NewWidgetLocation = LastPivotOverlayLocation;
+			FVector2f NewWidgetLocation = LastPivotOverlayLocation;
 
 			if (NewWidgetLocation.IsZero())
 			{
-				const FVector2D ActiveViewportSize = GetActiveViewportSize(ActiveLevelViewport);
+				const FVector2f ActiveViewportSize = GetActiveViewportSize(ActiveLevelViewport);
 				NewWidgetLocation.X = ActiveViewportSize.X / 2.0f;
 				NewWidgetLocation.Y = 50.0f;			
 
 			}
-			UpdatePivotOverlayLocation(NewWidgetLocation, ActiveLevelViewport);
+			UpdatePivotOverlayLocation(FVector2D(NewWidgetLocation), ActiveLevelViewport);
 
 			TAttribute<FMargin> MarginPadding = TAttribute<FMargin>::Create(TAttribute<FMargin>::FGetter::CreateUObject(this, &USequencerPivotTool::GetPivotOverlayPadding));
 
@@ -1285,8 +1285,8 @@ void USequencerPivotTool::TryShowPivotOverlay()
 
 void USequencerPivotTool::UpdatePivotOverlayLocation(const FVector2D InLocation, TSharedPtr<IAssetViewport> ActiveLevelViewport)
 {
-	const FVector2D ActiveViewportSize = GetActiveViewportSize(ActiveLevelViewport);
-	FVector2D ScreenPos = InLocation;
+	const FVector2f ActiveViewportSize = GetActiveViewportSize(ActiveLevelViewport);
+	FVector2f ScreenPos = UE::Slate::CastToVector2f(InLocation);
 
 	const float EdgeFactor = 0.97f;
 	const float MinX = ActiveViewportSize.X * (1 - EdgeFactor);
diff --git a/Engine/Plugins/Editor/SequencerAnimTools/Source/SequencerAnimTools/Public/SequencerAnimEditPivotTool.h b/Engine/Plugins/Editor/SequencerAnimTools/Source/SequencerAnimTools/Public/SequencerAnimEditPivotTool.h
--- a/Engine/Plugins/Editor/SequencerAnimTools/Source/SequencerAnimTools/Public/SequencerAnimEditPivotTool.h
+++ b/Engine/Plugins/Editor/SequencerAnimTools/Source/SequencerAnimTools/Public/SequencerAnimEditPivotTool.h
@@ -230,7 +230,7 @@ protected:
 
 	FMargin GetPivotOverlayPadding() const;
 	
-	static FVector2D LastPivotOverlayLocation;
+	static FVector2f LastPivotOverlayLocation;
 	TSharedPtr<SWidget> PivotWidget;
 public:
 	void TryRemovePivotOverlay();
diff --git a/Engine/Source/Editor/EditorFramework/Public/Toolkits/IToolkitHost.h b/Engine/Source/Editor/EditorFramework/Public/Toolkits/IToolkitHost.h
--- a/Engine/Source/Editor/EditorFramework/Public/Toolkits/IToolkitHost.h
+++ b/Engine/Source/Editor/EditorFramework/Public/Toolkits/IToolkitHost.h
@@ -63,9 +63,16 @@ public:
 	virtual void RemoveViewportOverlayWidget(TSharedRef<SWidget> InOverlaidWidget, TSharedPtr<IAssetViewport> InViewport = nullptr) {}
 
 	/** Return the size of the active viewport, accounting for DPI Scale */
-	virtual FVector2D GetActiveViewportSize() 
+	UE_DEPRECATED(5.6, "Viewport sizes are represented in float, please update your callsites to use GetActiveViewportWidgetSize and use FVector2f for slate sizes.")
+	virtual FVector2D GetActiveViewportSize() final
 	{
-		return FVector2D(EForceInit::ForceInitToZero);
+		return FVector2D(GetActiveViewportWidgetSize());
+	}
+
+	/** Return the size of the active viewport, accounting for DPI Scale */
+	virtual FVector2f GetActiveViewportWidgetSize() const
+	{
+		return FVector2f::ZeroVector;
 	}
 
 	/** Gets a multicast delegate which is executed whenever the toolkit host's active viewport changes. */
diff --git a/Engine/Source/Editor/LevelEditor/Private/SLevelEditor.cpp b/Engine/Source/Editor/LevelEditor/Private/SLevelEditor.cpp
--- a/Engine/Source/Editor/LevelEditor/Private/SLevelEditor.cpp
+++ b/Engine/Source/Editor/LevelEditor/Private/SLevelEditor.cpp
@@ -532,7 +532,7 @@ void SLevelEditor::Tick(const FGeometry& AllottedGeometry, const double InCurren
 	}
 }
 
-TSharedPtr<SLevelViewport> SLevelEditor::GetActiveViewport()
+TSharedPtr<SLevelViewport> SLevelEditor::GetActiveViewport() const
 {
 	// The first visible viewport
 	TSharedPtr<SLevelViewport> FirstVisibleViewport;
@@ -2251,7 +2251,7 @@ void SLevelEditor::RemoveViewportOverlayWidget(TSharedRef<SWidget> InWidget, TSh
 	}
 }
 	
-FVector2D SLevelEditor::GetActiveViewportSize()
+FVector2f SLevelEditor::GetActiveViewportWidgetSize() const
 {
 	if (TSharedPtr<SLevelViewport> ActiveViewport = GetActiveViewport())
 	{
@@ -2265,9 +2265,9 @@ FVector2D SLevelEditor::GetActiveViewportSize()
 		const float InvDpiScale = 1.0f / ActiveViewport->GetViewportClient()->GetDPIScale();
 		const float MaxX = SceneView->UnscaledViewRect.Width() * InvDpiScale;
 		const float MaxY = SceneView->UnscaledViewRect.Height() * InvDpiScale;
-		return FVector2D(MaxX, MaxY);
+		return FVector2f(MaxX, MaxY);
 	}
-	return FVector2D(0, 0);
+	return FVector2f::ZeroVector;
 }
 
 #undef LOCTEXT_NAMESPACE
diff --git a/Engine/Source/Editor/LevelEditor/Private/SLevelEditor.h b/Engine/Source/Editor/LevelEditor/Private/SLevelEditor.h
--- a/Engine/Source/Editor/LevelEditor/Private/SLevelEditor.h
+++ b/Engine/Source/Editor/LevelEditor/Private/SLevelEditor.h
@@ -66,7 +66,7 @@ public:
 	 *
 	 * @return The active viewport.  If multiple are active it returns the first one               
 	 */
-	TSharedPtr<class SLevelViewport> GetActiveViewport();
+	TSharedPtr<class SLevelViewport> GetActiveViewport() const;
 
 	/** ILevelEditor interface */
 	virtual const UTypedElementSelectionSet* GetElementSelectionSet() const override;
@@ -136,7 +136,7 @@ public:
 	virtual void RemoveViewportOverlayWidget(TSharedRef<SWidget>, TSharedPtr<IAssetViewport> InViewport = nullptr) override; 
 
 
-	virtual FVector2D GetActiveViewportSize() override;
+	virtual FVector2f GetActiveViewportWidgetSize() const override;
 	/* Tick to check the ActiveViewport */
 	virtual void Tick(const FGeometry& AllottedGeometry, const double InCurrentTime, const float InDeltaTime) override;
 
