Anim In Engine: Motion trails 2.0.

diff --git a/Engine/Plugins/Editor/SequencerAnimTools/Source/SequencerAnimTools/Private/Sequencer/MovieSceneTransformTrail.cpp b/Engine/Plugins/Editor/SequencerAnimTools/Source/SequencerAnimTools/Private/Sequencer/MovieSceneTransformTrail.cpp
--- a/Engine/Plugins/Editor/SequencerAnimTools/Source/SequencerAnimTools/Private/Sequencer/MovieSceneTransformTrail.cpp
+++ b/Engine/Plugins/Editor/SequencerAnimTools/Source/SequencerAnimTools/Private/Sequencer/MovieSceneTransformTrail.cpp
@@ -3,17 +3,25 @@
 #include "MovieSceneTransformTrail.h"
 #include "Engine/Engine.h"
 #include "TrailHierarchy.h"
+#include "SequencerTrailHierarchy.h"
 
 #include "Channels/MovieSceneChannelProxy.h"
 #include "Evaluation/MovieSceneEvaluationTemplateInstance.h"
 #include "Sequencer/MovieSceneControlRigParameterSection.h"
+#include "Sequencer/MovieSceneControlRigParameterTrack.h"
+#include "Tracks/MovieScene3DTransformTrack.h"
 
+#include "ControlRigObjectBinding.h"
+#include "Tools/MotionTrailOptions.h"
 
 #include "MovieSceneToolHelpers.h"
 #include "ActorForWorldTransforms.h"
 #include "Evaluation/MovieScenePlayback.h"
 #include "ScopedTransaction.h"
 #include "Sequencer/MotionTrailMovieSceneKey.h"
+#include "Framework/Application/SlateApplication.h"
+#include "PrimitiveDrawInterface.h"
+#include "Components/SkeletalMeshComponent.h"
 
 #define LOCTEXT_NAMESPACE "SequencerAnimTools"
 
@@ -22,6 +30,158 @@ namespace UE
 namespace SequencerAnimTools
 {
 
+using namespace UE::AIE;
+
+
+/****************************************************************************************************************
+*
+*  FCachedTracksAndSections
+*
+*******************************************************************************************************************/
+
+bool FCachedTracksAndSections::IsCacheEqual(const TArray <TWeakObjectPtr<UMovieSceneTrack>>& InTracks)
+{
+	if (Tracks.Num() != InTracks.Num())
+	{
+		CreateFrom(InTracks);
+		return false;
+	}
+	else
+	{
+		for (const TWeakObjectPtr<UMovieSceneTrack>& WeakTrack : InTracks)
+		{
+			if (WeakTrack.IsValid())
+			{
+				if (FCachedTrackAndSections* CachedTrackAndSections = Tracks.Find(WeakTrack))
+				{
+					if (WeakTrack->GetSignature() != CachedTrackAndSections->Track.CachedGuid)
+					{
+						CreateFrom(InTracks);
+						return false;
+					}
+					const TArray<UMovieSceneSection*>& Sections = WeakTrack->GetAllSections();
+					if (Sections.Num() != CachedTrackAndSections->Sections.Num())
+					{
+						CreateFrom(InTracks);
+						return false;
+					}
+					for (const UMovieSceneSection* Section : Sections)
+					{
+						if (FCachedGuidsPerSignedObject* SignedGuid = CachedTrackAndSections->Sections.Find(Section))
+						{
+							if (Section->GetSignature() != SignedGuid->CachedGuid)
+							{
+								CreateFrom(InTracks);
+								return false;
+							}
+						}
+						else
+						{
+							CreateFrom(InTracks);
+							return false;
+						}
+					}
+				}
+				else
+				{
+					CreateFrom(InTracks);
+					return false;
+				}
+			}
+		}
+		return true;
+	}
+}
+
+void  FCachedTracksAndSections::CreateFrom(const TArray <TWeakObjectPtr<UMovieSceneTrack>>& InTracks)
+{
+	Tracks.Reset();
+	for (const TWeakObjectPtr<UMovieSceneTrack>& WeakTrack : InTracks)
+	{
+		if (WeakTrack.IsValid())
+		{
+			FCachedTrackAndSections Track;
+			Track.Track.SignedObject = WeakTrack;
+			Track.Track.CachedGuid = WeakTrack->GetSignature();
+			const TArray<UMovieSceneSection*>& Sections = WeakTrack->GetAllSections();
+			for (const UMovieSceneSection* Section : Sections)
+			{
+				if (Section)
+				{
+					FCachedGuidsPerSignedObject SignedObjectGuid;
+					SignedObjectGuid.SignedObject = Section;
+					SignedObjectGuid.CachedGuid = Section->GetSignature();
+					Track.Sections.Add(Section,SignedObjectGuid);
+
+				}
+			}
+			Tracks.Add(WeakTrack,Track);
+		}
+	}
+}
+
+UMovieSceneSection* FCachedTracksAndSections::GetSection() const
+{
+	for (const TPair<TWeakObjectPtr<UMovieSceneTrack>, FCachedTrackAndSections>& Pair : Tracks)
+	{
+		if (Pair.Key.IsValid())
+		{
+			UMovieSceneSection* Section = Pair.Key.Get()->GetSectionToKey();
+			if (Section)
+			{
+				return Section;
+			}
+			else if (Pair.Key.Get()->GetAllSections().Num() > 0)
+			{
+				return Pair.Key.Get()->GetAllSections()[0];
+			}
+		}
+	}
+	return nullptr;
+}
+/****************************************************************************************************************
+*
+*   Optional Parent Space
+*
+*****************************************************************************************************************/
+
+void FOptionalParentSpace::SetSpace(AActor* InActor, const FName& InComponentName)
+{
+	if (InActor)
+	{
+		bIsValid = true;
+		ParentSpace.Actor.Actor = InActor;
+		if (InComponentName != NAME_None)
+		{
+			TInlineComponentArray<USceneComponent*> Components(InActor);
+
+			for (USceneComponent* Component : Components)
+			{
+				if (Component->GetFName() == InComponentName)
+				{
+					ParentSpace.Actor.Component = Component;
+					break;
+				}
+			}
+		}
+
+		if (ParentSpace.WorldTransforms.IsValid() == false)
+		{
+			ParentSpace.WorldTransforms = MakeShared<FArrayOfTransforms>();
+		}
+		if (ParentSpace.ParentTransforms.IsValid() == false)
+		{
+			ParentSpace.ParentTransforms = MakeShared<FArrayOfTransforms>();
+		}
+	}
+}
+
+void FOptionalParentSpace::ClearSpace()
+{
+	bIsValid = false;
+	ParentSpace.Actor.Actor = nullptr;
+}
+
 /****************************************************************************************************************
 *
 *   Base Movie Scene Transform Trail
@@ -29,85 +189,147 @@ namespace SequencerAnimTools
 *****************************************************************************************************************/
 
 
-FMovieSceneTransformTrail::FMovieSceneTransformTrail(USceneComponent* InOwner, const FLinearColor& InColor, const bool bInIsVisible, TWeakObjectPtr<UMovieSceneTrack> InWeakTrack, TSharedPtr<class ISequencer> InSequencer)
+FMovieSceneTransformTrail::FMovieSceneTransformTrail(const FGuid& InBindingID,USceneComponent* InOwner,  const bool bInIsVisible, TWeakObjectPtr<UMovieSceneTrack> InWeakTrack, TSharedPtr<ISequencer> InSequencer)
 	: FTrail(InOwner)
 	, WeakSequencer(InSequencer)
-	, CachedEffectiveRange(TRange<double>::Empty())
 	, bIsSelected(false)
-	, TrajectoryCache()
-	, LastTransformSectionSig(FGuid())
-	, CachedHierarchyGuid()
-	, WeakTrack(InWeakTrack)
+	, bIsOffsetMode(false)
+	, bIsTracking(false)
+	, BindingID(InBindingID)
+	, MainTrack(InWeakTrack)
 {
-	TrajectoryCache = MakeUnique<FArrayTrajectoryCache>(0.01, GetEffectiveSectionRange(0));
 	KeyTool = MakeUnique< FMotionTraiMovieScenelKeyTool>(this);
-	DrawInfo = MakeUnique<FTrajectoryDrawInfo>(InColor, TrajectoryCache.Get());
+
+	//new
+	ArrayOfTransforms = MakeShared<FArrayOfTransforms>();
+	ParentArrayOfTransforms = MakeShared<FArrayOfTransforms>();
+
+	DrawInfo = MakeUnique<FTrajectoryDrawInfo>(UMotionTrailToolOptions::GetTrailOptions()->TrailStyle, UMotionTrailToolOptions::GetTrailOptions()->DefaultColor, ArrayOfTransforms, ParentSpace.ParentSpace.WorldTransforms);
 }
 
-UMovieSceneSection* FMovieSceneTransformTrail::GetSection() const
+void FMovieSceneTransformTrail::HasStartedEvaluating()
 {
-	if (WeakTrack.IsValid())
+	PreviousCachedDrawData = CachedDrawData;
+}
+
+void FMovieSceneTransformTrail::UpdateFinished(const TRange<FFrameNumber>& Range, const TArray<int32>& IndicesToCalcluate)
+{
+	FTrail::UpdateFinished(Range, IndicesToCalcluate);
+	//if we are in a parent space we need to put transforms in that space
+	if (ParentSpace.bIsValid && IndicesToCalcluate.Num() > 0)
 	{
-		UMovieSceneSection* Section = WeakTrack.Get()->GetSectionToKey();
-		if (Section)
-		{
-			return Section;
-		}
-		else if( WeakTrack.Get()->GetAllSections().Num() >0)
+		for (int32 Count = 0; Count < IndicesToCalcluate.Num(); ++Count)
 		{
-			return WeakTrack.Get()->GetAllSections()[0];
+			int32 Index = IndicesToCalcluate[Count];
+			ArrayOfTransforms->Transforms[Index] = ArrayOfTransforms->Transforms[Index].GetRelativeTransform(ParentSpace.ParentSpace.WorldTransforms->Transforms[Index]);
+			ParentArrayOfTransforms->Transforms[Index] = ParentArrayOfTransforms->Transforms[Index].GetRelativeTransform(ParentSpace.ParentSpace.WorldTransforms->Transforms[Index]);
 		}
+
+	}
+
+	if (bIsTracking == false)
+	{
+		KeyTool->DirtyKeyTransforms();
+		KeyTool->UpdateKeys();
 	}
-	return  nullptr;
 }
 
-bool FMovieSceneTransformTrail::GetEditedTimes(const FTrailHierarchy* TrailHierarchy,const FFrameNumber& LastFrame, TArray<FFrameNumber>& OutEditedTimes)
+FTransform FMovieSceneTransformTrail::GetParentSpaceTransform()  const
 {
-	if (EditedTimes.IsSet())
+	if (ParentSpace.bIsValid == true)
 	{
-		CalculateEditedTimes(TrailHierarchy,LastFrame);
-		if (EditedTimes.IsSet())
+		if (TSharedPtr<ISequencer> Sequencer = WeakSequencer.Pin())
 		{
-			OutEditedTimes = EditedTimes.GetValue();
-			return true;
+			FFrameTime Time = Sequencer->GetLocalTime().Time;
+			USceneComponent* ParentComponent = ParentSpace.ParentSpace.Actor.Component.IsValid() ? ParentSpace.ParentSpace.Actor.Component.Get() : ParentSpace.ParentSpace.Actor.Actor->GetRootComponent();
+			FTransform ParentSpaceTM = ParentComponent ? ParentComponent->GetComponentToWorld() : FTransform::Identity;				
+			if (PreviousParentSpaceTM.Equals(ParentSpaceTM) == false)
+			{
+				CachedDrawData.PointsToDraw.SetNum(0);
+				PreviousParentSpaceTM = ParentSpaceTM;
+			}
+			return ParentSpaceTM;
 		}
 	}
-	return false;
+	return FTransform::Identity;
 }
 
-void FMovieSceneTransformTrail::CalculateEditedTimes(const FTrailHierarchy* TrailHierarchy, const FFrameNumber& LastFrame)
+void FMovieSceneTransformTrail::AddImportantTimes(TSet<FFrameNumber>& InOutImportantTimes)
 {
-	if (bIsSelected || KeyTool->IsAllSelected())
+	TArray<FFrameNumber> SelectedKeyTimes = KeyTool->SelectedKeyTimes();
 	{
-		EditedTimes.Reset();
+		for (const FFrameNumber& FrameNumber : SelectedKeyTimes)
+		{
+			InOutImportantTimes.Add(FrameNumber);
+		}
 	}
-	else
+}
+
+UMovieSceneSection* FMovieSceneTransformTrail::GetSection() const
+{
+	UMovieSceneSection* Section = nullptr;
+	if (MainTrack.IsValid())
 	{
-		const FFrameNumber Step = TrailHierarchy->GetFramesPerSegment();
-		TArray<FFrameNumber> SelectedTimes  = KeyTool->SelectedKeyTimes();
-		EditedTimes = KeyTool->GetTimesFromModifiedTimes(SelectedTimes,LastFrame,Step);
-		if (EditedTimes.GetValue().Num() == 0)
+		Section = MainTrack->GetSectionToKey();
+		if (!Section && MainTrack->GetAllSections().Num() > 0)
 		{
-			EditedTimes.Reset();
+			Section = MainTrack->GetAllSections()[0];
 		}
 	}
+	return Section;
 }
 
 void FMovieSceneTransformTrail::ForceEvaluateNextTick()
 {
 	FTrail::ForceEvaluateNextTick();
 	KeyTool->DirtyKeyTransforms();
+	ClearCachedData();
 }
 
-void FMovieSceneTransformTrail::UpdateKeysInRange(const TRange<double>& ViewRange)
+//if actor is not a sequencer actor then it's static so we need to update the trails
+void FMovieSceneTransformTrail::ActorChanged(AActor* InActor)
 {
-	KeyTool->DirtyKeyTransforms();
-	KeyTool->UpdateKeysInRange(ViewRange);
+	if (TransformDependencies.NonSequencerActors.Contains(InActor))
+	{
+		ForceEvaluateNextTick();
+	}
+}
+
+bool FMovieSceneTransformTrail::TrailHasChanged()
+{
+	if(TSharedPtr<ISequencer> SequencerPtr = WeakSequencer.Pin())
+	{ 
+		if (USceneComponent* Component = Cast<USceneComponent>(GetOwner()))
+		{
+			UMovieScene* MovieScene = SequencerPtr->GetFocusedMovieSceneSequence()->GetMovieScene();
+			if (FMovieSceneBinding* Binding = MovieScene->FindBinding(BindingID))
+			{
+				AActor* Actor = Component->GetOwner();
+				UE::AIE::FSequencerTransformDependencies NewTransformDependencies;
+				TArray<UMovieSceneTrack*> Tracks = FSequencerTransformDependencies::GetTransformAffectingTracks(SequencerPtr->GetFocusedMovieSceneSequence()->GetMovieScene(),
+					*Binding);
+
+				NewTransformDependencies.CalculateDependencies(WeakSequencer.Pin().Get(), Actor, Tracks);
+				if (TransformDependencies.Compare(NewTransformDependencies) == false)
+				{
+					TransformDependencies.CopyFrom(NewTransformDependencies);
+					return true;
+				}
+			}
+		
+		}
+	}
+	return false;
 }
 
-ETrailCacheState FMovieSceneTransformTrail::UpdateTrail(const FSceneContext& InSceneContext)
+FTrailCurrentStatus FMovieSceneTransformTrail::UpdateTrail(const FNewSceneContext& InSceneContext)
 {
+	FTrailCurrentStatus Status;
+	
 	CachedHierarchyGuid = InSceneContext.YourNode;
+	
+	const bool bTrackUnchanged = (InSceneContext.bCheckForChange || TransformDependencies.IsEmpty() ) ? !TrailHasChanged() : true;
+
 	UMovieSceneSection* Section = GetSection();
 	TSharedPtr<ISequencer> Sequencer = WeakSequencer.Pin();
 
@@ -119,50 +341,221 @@ ETrailCacheState FMovieSceneTransformTrail::UpdateTrail(const FSceneContext& InS
 
 	if (!Sequencer || !Section || !SequencerBinding.IsValid())
 	{
-		return ETrailCacheState::Dead;
+		Status.CacheState = ETrailCacheState::Dead;
+		return Status;
 	}
 	
-	const bool bTrackUnchanged = Section->GetSignature() == LastTransformSectionSig;
+	FSequencerTrailHierarchy* SequencerTrailHierarchy = static_cast<FSequencerTrailHierarchy*>(InSceneContext.TrailHierarchy);
+	CurrentFramesInfo = SequencerTrailHierarchy->GetCurrentFramesInfo();
+	if (!bTrackUnchanged || bForceEvaluateNextTick)
+	{
+		KeyTool->OnSectionChanged();
+		Status.CacheState = ETrailCacheState::Stale;
+		bForceEvaluateNextTick = false;
 
-	if (!bTrackUnchanged  || bForceEvaluateNextTick)
+		UpdateNeedsEvaluation(InSceneContext.YourNode, SequencerTrailHierarchy);
+	}
+	else
 	{
+		Status.CacheState = ETrailCacheState::UpToDate;
 
-		KeyTool->OnSectionChanged();
-		TrajectoryCache->UpdateCacheTimes(InSceneContext.EvalTimes);
+	}
+	KeyTool->UpdateViewRange(CurrentFramesInfo->ViewRange);
 
-		CacheState = ETrailCacheState::Stale;
-		bForceEvaluateNextTick = false;
-		LastTransformSectionSig = Section->GetSignature();
+	return Status;
+}
 
-		UpdateCacheTimes(InSceneContext);
-				
+void FMovieSceneTransformTrail::Interp(const FFrameNumber& CurrentFrame, FTransform& OutTransform, FTransform& OutParentTransform)
+{
+	OutTransform = ArrayOfTransforms->Interp(CurrentFrame, CurrentFramesInfo->TransformIndices, CurrentFramesInfo->CurrentFrames);
+	OutParentTransform = ParentArrayOfTransforms->Interp(CurrentFrame, CurrentFramesInfo->TransformIndices, CurrentFramesInfo->CurrentFrames);
+}
+
+
+void FMovieSceneTransformTrail::ClearCachedData()
+{
+	CachedDrawData.PointsToDraw.SetNum(0);
+	CachedDrawData.Color.SetNum(0);
+}
+
+void FMovieSceneTransformTrail::GetColor(const FFrameNumber& CurrentTime, FColorState& State)
+{
+	if (State.GetStyle() == EMotionTrailTrailStyle::HeatMap)
+	{
+		if (State.bFirstFrame == true)
+		{
+			State.bFirstFrame = false;
+			State.StartFrame = CurrentTime;
+			float MinVelocity = FLT_MAX;
+			float MaxVelocity = FLT_MIN;
+			TArray<float> Velocities;
+			Velocities.SetNum(CurrentFramesInfo->TransformIndices.Num());
+			int32 TransformIndex = CurrentFramesInfo->TransformIndices[0];
+			FVector LastPosition = ArrayOfTransforms->Transforms[TransformIndex].GetLocation();
+			float AverageValue = 0.0f;
+			for (int32 Index = 1; Index < CurrentFramesInfo->TransformIndices.Num(); ++Index)
+			{
+				TransformIndex = CurrentFramesInfo->TransformIndices[Index];
+				const FVector Position = ArrayOfTransforms->Transforms[TransformIndex].GetLocation();
+				const float Velocity = (Position - LastPosition).Length();
+				LastPosition = Position;
+				if (Velocity > MaxVelocity)
+				{
+					MaxVelocity = Velocity;
+				}
+				if (Velocity < MinVelocity)
+				{
+					MinVelocity = Velocity;
+				}
+				Velocities[Index - 1] = Velocity;
+				AverageValue += Velocity;
+			}
+
+			AverageValue /= (float)Velocities.Num();
+			if (FMath::IsNearlyZero(AverageValue) == false)
+			{
+				CachedHeatMap.SetNum(CurrentFramesInfo->TransformIndices.Num());
+				const float LowValue = MinVelocity;//AverageValue * 0.5f;
+				const float HighValue = MaxVelocity;//AverageValue * 2.0f;
+				const float HighAverage = HighValue - AverageValue;
+				const float AverageLow = AverageValue - LowValue;
+				FColor Color;
+				for (int32 Index = 0; Index < Velocities.Num() - 1; ++Index)
+				{
+					float UpperWeight = FMath::Clamp(Velocities[Index] - AverageValue, 0.0f, HighAverage)/ HighAverage;
+					float LowerWeight = FMath::Clamp(AverageValue - Velocities[Index], 0.0f, AverageLow) / AverageLow;
+					Color.R = (uint8)(UpperWeight * 255);
+					Color.G = (uint8)((1.0 - (UpperWeight + LowerWeight)) * 255);
+					Color.B = (uint8)(LowerWeight * 255);
+					Color.A = Velocities[Index] >= LowValue ? 255 : 0;
+					CachedHeatMap[Index] = Color;
+				}
+			}
+			else //average speed is zero so not moving, no heat map color
+			{
+				CachedHeatMap.SetNum(0);
+			}
+
+		}
+		if (CachedHeatMap.Num() > 0)
+		{
+			const int32 Index = (CurrentTime.Value - State.StartFrame.Value) / State.TicksPerFrame.Value;
+			if (Index >= 0 && Index <= CachedHeatMap.Num() - 1)
+			{
+				State.CalculatedColor = CachedHeatMap[Index];
+			}
+			else
+			{
+				State.CalculatedColor = FLinearColor::White;
+			}
+		}
+		else
+		{
+			State.CalculatedColor = FLinearColor::White;
+		}
 	}
-	else 
+	else
 	{
-		TrajectoryCache->UpdateCacheTimes(InSceneContext.EvalTimes);
-		KeyTool->UpdateSelectedKeysTransform();
-		CacheState = ETrailCacheState::UpToDate;
+		FTrail::GetColor(CurrentTime, State);
 	}
-	return CacheState;
 }
 
+void FMovieSceneTransformTrail::ReadyToDrawTrail(FColorState& ColorState, const FCurrentFramesInfo* InCurrentFramesInfo, bool bIsEvaluating, bool bIsPinned)
+{
+	ColorState.ReadyForTrail(bIsPinned, GetDrawInfo()->GetStyle());
+	bool bCalculateColor = (ColorState.GetStyle() == EMotionTrailTrailStyle::Time);
+		
+	if (InCurrentFramesInfo->bViewRangeIsEvalRange == false || (CachedDrawData.PointsToDraw.Num() != CachedDrawData.Color.Num()) || (CachedDrawData.PointsToDraw.Num() != InCurrentFramesInfo->CurrentFrames.Num()))
+	{
+		GetTrajectoryPointsForDisplay(*InCurrentFramesInfo, bIsEvaluating, CachedDrawData.PointsToDraw,CachedDrawData.Frames);
+		bCalculateColor = true;
+	}
+	if (CachedDrawData.PointsToDraw.Num() > 1 && bCalculateColor)
+	{
+		CachedDrawData.Color.SetNum(CachedDrawData.PointsToDraw.Num());
+		for (int32 Idx = 1; Idx < CachedDrawData.PointsToDraw.Num(); Idx++)
+		{
+			const FFrameNumber CurFrame = CachedDrawData.Frames[Idx - 1];
+			GetColor(CurFrame, ColorState);
+			CachedDrawData.Color[Idx - 1] = ColorState.CalculatedColor;
+		}
+	}
+}
 void FMovieSceneTransformTrail::DrawHUD(const FSceneView* View, FCanvas* Canvas)
 {
 	KeyTool->DrawHUD(View, Canvas);
 }
 
+TOptional<FLinearColor> FMovieSceneTransformTrail::GetOverrideColor() const
+{
+	TOptional<FLinearColor> Color;
+	if (IsTrailSelected())
+	{
+		if (bIsOffsetMode)
+		{
+			Color = FLinearColor::Red;
+		}
+		else
+		{
+			Color = FLinearColor::Yellow;
+		}
+	}
+	return Color;
+}
+
+void FMovieSceneTransformTrail::InternalDrawTrail(const FGuid& Guid, const FSceneView* View, FPrimitiveDrawInterface* PDI, const FDrawCacheData& InDrawData, bool bHitTesting) const
+{
+	const TOptional<FLinearColor> OverrideColor = GetOverrideColor();
+	const float TrailThickness = UMotionTrailToolOptions::GetTrailOptions()->TrailThickness;
+	FVector LastPoint = InDrawData.PointsToDraw[0];
+
+	for (int32 Idx = 1; Idx < InDrawData.PointsToDraw.Num(); Idx++)
+	{
+		const FFrameNumber CurFrame = InDrawData.Frames[Idx - 1];
+		if (bHitTesting)
+		{
+			PDI->SetHitProxy(new HNewMotionTrailProxy(Guid, LastPoint, CurFrame));
+		}
+		const FVector CurPoint = InDrawData.PointsToDraw[Idx];
+		const FLinearColor CurColor = OverrideColor.IsSet() == false ? InDrawData.Color[Idx - 1] : OverrideColor.GetValue();
+		PDI->DrawLine(LastPoint, CurPoint, CurColor, SDPG_Foreground, TrailThickness);
+		LastPoint = CurPoint;
+		if (bHitTesting)
+		{
+			PDI->SetHitProxy(nullptr);
+		}
+	}
+}
+
 void FMovieSceneTransformTrail::Render(const FGuid& Guid, const FSceneView* View, FPrimitiveDrawInterface* PDI)
 {
+	if (CachedDrawData.PointsToDraw.Num() > 1)
+	{
+		const bool bHitTesting = PDI && PDI->IsHitTesting();
+		InternalDrawTrail(Guid, View, PDI, CachedDrawData, bHitTesting);
+	}
+	//render key tool
 	KeyTool->Render(Guid, View, PDI);
 }
 
-bool FMovieSceneTransformTrail::HandleAltClick(FEditorViewportClient* InViewportClient, HMotionTrailProxy* Proxy, FInputClick Click)
+void FMovieSceneTransformTrail::RenderEvaluating(const FGuid& Guid, const FSceneView* View, FPrimitiveDrawInterface* PDI)
+{
+	if (PreviousCachedDrawData.PointsToDraw.Num() > 0)
+	{
+		const bool bHitTesting = false;
+		InternalDrawTrail(Guid, View, PDI, PreviousCachedDrawData, bHitTesting);
+	}
+}
+
+bool FMovieSceneTransformTrail::HandleAltClick(FEditorViewportClient* InViewportClient, HNewMotionTrailProxy* Proxy, FInputClick Click)
 {
 	return false;
 }
 
 bool FMovieSceneTransformTrail::HandleClick(const FGuid& Guid, FEditorViewportClient* InViewportClient, HHitProxy* InHitProxy, FInputClick Click)
 {
+	bIsOffsetMode = false;
+
 	bool bKeySelected = KeyTool->HandleClick(Guid,InViewportClient, InHitProxy, Click);
 	if (bKeySelected)
 	{
@@ -173,7 +566,7 @@ bool FMovieSceneTransformTrail::HandleClick(const FGuid& Guid, FEditorViewportCl
 		return true;
 	}
 
-	if (HMotionTrailProxy* HitProxy = HitProxyCast<HMotionTrailProxy>(InHitProxy))
+	if (HNewMotionTrailProxy* HitProxy = HitProxyCast<HNewMotionTrailProxy>(InHitProxy))
 	{
 		if (HitProxy->Guid == Guid)
 		{
@@ -214,6 +607,35 @@ bool FMovieSceneTransformTrail::HandleClick(const FGuid& Guid, FEditorViewportCl
 	return false;
 }
 
+bool FMovieSceneTransformTrail::IsAnythingSelected(TArray<FVector>& OutVectorPositions) const
+{
+	if (bIsSelected)
+	{
+		const FTransform ParentSpaceTransform = GetParentSpaceTransform();
+		TArray<FTrailKeyInfo*> AllKeys;
+		KeyTool->GetAllKeys(AllKeys);
+		for (const FTrailKeyInfo* KeyInfo : AllKeys)
+		{
+			if (KeyInfo)
+			{
+				const FTransform Transform = KeyInfo->Transform * OffsetTransform * ParentSpaceTransform;
+				OutVectorPositions.Add(KeyInfo->Transform.GetLocation());
+			}
+		}
+		return true;
+	}
+	else
+	{
+		const int32 CurrentSelected = OutVectorPositions.Num();
+		KeyTool->GetSelectedKeyPositions(OutVectorPositions);
+		if (CurrentSelected != OutVectorPositions.Num())
+		{
+			return true;
+		}
+	}
+	return false;
+}
+
 bool FMovieSceneTransformTrail::IsAnythingSelected(FVector& OutVectorPosition) const
 {
 	if (bIsSelected)
@@ -237,6 +659,7 @@ void FMovieSceneTransformTrail::DeleteSelectedKeys()
 void FMovieSceneTransformTrail::SelectNone()
 {
 	KeyTool->ClearSelection();
+	bIsSelected = false;
 }
 
 bool FMovieSceneTransformTrail::IsAnythingSelected() const
@@ -249,18 +672,87 @@ bool FMovieSceneTransformTrail::IsTrailSelected() const
 	return bIsSelected;
 }
 
+void FMovieSceneTransformTrail::ClearOffsetTransform()
+{
+	OffsetTransform = FTransform::Identity;
+	if (bIsOffsetMode)
+	{
+		bIsOffsetMode = false;
+		bIsSelected = false;
+		ClearCachedData();
+	}
+}
+
+void FMovieSceneTransformTrail::SetOffsetMode()
+{
+	if (bIsOffsetMode == false)
+	{
+		bIsOffsetMode = true;
+		SelectNone();
+		bIsSelected = true;
+		//set selected position aka gizmo location  to be the average of the points drawn
+		if (CachedDrawData.PointsToDraw.Num() > 0)
+		{
+			FVector Min(DBL_MAX, DBL_MAX, DBL_MAX);
+			FVector Max(DBL_MIN, DBL_MIN, DBL_MIN);
+			for (const FVector& Vector : CachedDrawData.PointsToDraw)
+			{
+				for (int32 Index = 0; Index <= 2; ++Index)
+				{
+					if (Vector[Index] < Min[Index])
+					{
+						Min[Index] = Vector[Index];
+					}
+					if (Vector[Index] > Max[Index])
+					{
+						Max[Index] = Vector[Index];
+					}
+				}
+			}
+			for (int32 Index = 0; Index <= 2; ++Index)
+			{
+				SelectedPos[Index] = (Max[Index] + Min[Index]) * 0.5;
+			}
+		}
+	}
+
+}
 bool FMovieSceneTransformTrail::StartTracking()
 {
+	bIsTracking = true;
+	//if in offset mode this happened from the tool bar menu,
+	// if not we check fo shift is held we go into offset mode
+	if (!bIsOffsetMode)
+	{
+		bIsOffsetMode = FSlateApplication::Get().GetModifierKeys().IsShiftDown();
+	}
+	
 	return false;
 }
 
-bool FMovieSceneTransformTrail::ApplyDelta(const FVector& Pos, const FRotator& Rot, const FVector& WidgetLocation)
+bool FMovieSceneTransformTrail::ApplyDelta(const FVector& Pos, const FRotator& Rot, const FVector& WidgetLocation, bool bApplyToOffset)
 {
-	return false;;
+	if (bApplyToOffset || bIsOffsetMode)
+	{
+		if (IsTrailSelected())
+		{
+			OffsetTransform.SetLocation(OffsetTransform.GetLocation() + Pos);
+			int32 Index = UMotionTrailToolOptions::GetTrailOptions()->GetIndexFromGuid(CachedHierarchyGuid);
+			if (Index != INDEX_NONE)
+			{
+				UMotionTrailToolOptions::GetTrailOptions()->SetHasOffset(Index,true);
+			}
+			SelectedPos += Pos;
+			ClearCachedData();
+			return true;
+		}
+	}
+	return false;
 }
 
 bool FMovieSceneTransformTrail::EndTracking()
 {
+	bIsTracking = false;
 	return false;
 }
 
@@ -285,49 +777,46 @@ TArray<FFrameNumber> FMovieSceneTransformTrail::GetSelectedKeyTimes() const
 	return KeyTool->SelectedKeyTimes();
 }
 
-TRange<double> FMovieSceneTransformTrail::GetEffectiveSectionRange(int32 ChannelOffset) const
+void FMovieSceneTransformTrail::UpdateNeedsEvaluation(const FGuid& InTrailGuid, FSequencerTrailHierarchy* SequencerHierarchy)
 {
-	UMovieSceneSection* TransformSection = GetSection();
-	if (!TransformSection)
+	if (ParentSpace.bIsValid)
 	{
-		return TRange<double>();
+		FActorForWorldTransforms Actors;
+		Actors.Actor = ParentSpace.ParentSpace.Actor.Actor;
+		Actors.Component = ParentSpace.ParentSpace.Actor.Component;
+		SequencerHierarchy->EvaluateActor(InTrailGuid, Actors, ParentSpace.ParentSpace.WorldTransforms, ParentSpace.ParentSpace.ParentTransforms);
 	}
-	TSharedPtr<ISequencer> Sequencer = WeakSequencer.Pin();
-	check(TransformSection && Sequencer);
-
-	TRange<FFrameNumber> EffectiveRange = TRange<FFrameNumber>::Empty();
-
-	int32 MaxChannel = (int32)ETransformChannel::TranslateZ; //0,1,2 are the position channels only onces we care about.
-	TArrayView<FMovieSceneDoubleChannel*> DoubleChannels = TransformSection->GetChannelProxy().GetChannels<FMovieSceneDoubleChannel>();
-	TArrayView<FMovieSceneFloatChannel*> FloatChannels = TransformSection->GetChannelProxy().GetChannels<FMovieSceneFloatChannel>();
+}
 
-	if (DoubleChannels.Num() > MaxChannel)
+void FMovieSceneTransformTrail::FMovieSceneTransformTrailState::SaveFromTrail(const FMovieSceneTransformTrail* InTrail)
+{
+	if (InTrail)
 	{
-		for (int32 ChannelIdx = ChannelOffset; ChannelIdx <= (ChannelOffset + MaxChannel); ChannelIdx++)
+		BindingID = InTrail->BindingID;
+		WeakSequencer = InTrail->WeakSequencer;
+		MainTrack = InTrail->MainTrack;
+		OffsetTransform = InTrail->OffsetTransform;
+		ParentSpace = InTrail->ParentSpace;
+		if (InTrail->GetDrawInfo())
 		{
-			FMovieSceneDoubleChannel* Channel = DoubleChannels[ChannelIdx];
-			EffectiveRange = TRange<FFrameNumber>::Hull(EffectiveRange, Channel->ComputeEffectiveRange());
+			Color = InTrail->GetDrawInfo()->GetColor();
+			PinnedStyle = InTrail->GetDrawInfo()->GetStyle();
 		}
 	}
-	if (FloatChannels.Num() > MaxChannel)
+}
+
+void FMovieSceneTransformTrail::FMovieSceneTransformTrailState::SetToTrail(FMovieSceneTransformTrail* InTrail) const
+{
+	if (InTrail)
 	{
-		for (int32 ChannelIdx = ChannelOffset; ChannelIdx <= (ChannelOffset + MaxChannel); ChannelIdx++)
+		if (InTrail->GetDrawInfo())
 		{
-			FMovieSceneFloatChannel* Channel = FloatChannels[ChannelIdx];
-			EffectiveRange = TRange<FFrameNumber>::Hull(EffectiveRange, Channel->ComputeEffectiveRange());
+			InTrail->GetDrawInfo()->SetColor(Color);
+			InTrail->GetDrawInfo()->SetStyle(PinnedStyle);
 		}
+		InTrail->OffsetTransform = OffsetTransform;
+		InTrail->ParentSpace = ParentSpace;
 	}
-
-	EffectiveRange = TRange<FFrameNumber>::Intersection(EffectiveRange, TransformSection->GetRange());
-
-	TRange<double> SectionRangeSeconds = TRange<double>(
-		Sequencer->GetFocusedTickResolution().AsSeconds(EffectiveRange.GetLowerBoundValue()),
-		Sequencer->GetFocusedTickResolution().AsSeconds(EffectiveRange.GetUpperBoundValue())
-	);
-		
-	// TODO: clip by movie scene range? try movie keys outside of movie scene range
-
-	return SectionRangeSeconds;
 }
 
 /****************************************************************************************************************
@@ -336,61 +825,33 @@ TRange<double> FMovieSceneTransformTrail::GetEffectiveSectionRange(int32 Channel
 *
 *****************************************************************************************************************/
 
-
-void FMovieSceneComponentTransformTrail::UpdateCacheTimes(const FSceneContext& InSceneContext)
+void FMovieSceneComponentTransformTrail::UpdateNeedsEvaluation(const FGuid& InTrailGuid, FSequencerTrailHierarchy* SequencerHierarchy)
 {
-	TSharedPtr<ISequencer> Sequencer = WeakSequencer.Pin();
-	check(Sequencer);
-	FFrameRate TickResolution = Sequencer->GetFocusedTickResolution();
-	const FTrailEvaluateTimes& EvaluateTimes = InSceneContext.EvalTimes;
-	TArray<FFrameNumber> Frames;
-	bool bUseEditedTimes = false;
-	if (EditedTimes.IsSet())
+	if (Component.IsValid())
 	{
-		const FFrameTime LastFrame = EvaluateTimes.EvalTimes[EvaluateTimes.EvalTimes.Num() - 1] * TickResolution;
+		FMovieSceneTransformTrail::UpdateNeedsEvaluation(InTrailGuid, SequencerHierarchy);
+		FActorForWorldTransforms Actors;
+		Actors.Actor = Component->GetOwner();
+		Actors.Component = Component;
+		SequencerHierarchy->EvaluateActor(InTrailGuid, Actors, ArrayOfTransforms, ParentArrayOfTransforms);
 
-		CalculateEditedTimes(InSceneContext.TrailHierarchy,LastFrame.RoundToFrame());
-		Frames = EditedTimes.GetValue();
-		bUseEditedTimes = true;
-	}
-	else
-	{
-		//Todo mz make time in frames
-		int32 Index = 0;
-		Frames.SetNum(EvaluateTimes.EvalTimes.Num());
-		for (const double Time : EvaluateTimes.EvalTimes)
-		{
-			const FFrameTime TickTime = Time * TickResolution;
-			Frames[Index++] = TickTime.RoundToFrame();
-		}
 	}
+}
 
-	TArray<FTransform> OutWorldTransforms;
-	FActorForWorldTransforms Actors;
-	Actors.Actor = Component->GetTypedOuter<AActor>();
-	Actors.Component = Component;
-	MovieSceneToolHelpers::GetActorWorldTransforms(GetSequencer().Get(), Actors, Frames, OutWorldTransforms);
-	TArray<FTransform> ParentCache;
-	if (EditedTimes.IsSet())
-	{
-		for (int32 Index = 0; Index < Frames.Num(); ++Index)
-		{
-			const FFrameNumber& FrameNumber = Frames[Index];
-			double Sec = TickResolution.AsSeconds(FFrameTime(FrameNumber));
-			GetTrajectoryTransforms()->Set(Sec, OutWorldTransforms[Index]);
-		}
-	}
-	else
+FText FMovieSceneComponentTransformTrail::GetName() const
+{
+	if (Component.IsValid())
 	{
-		GetTrajectoryTransforms()->SetTransforms(OutWorldTransforms, ParentCache);
+		FString ActorLabel = Component->GetOwner()->GetActorLabel(false);
+		return FText::FromString(ActorLabel);
 	}
-
-	KeyTool->UpdateKeysInRange(InSceneContext.TrailHierarchy->GetViewRange());
-	EditedTimes.Reset();
+	return FText();
 }
 
 bool FMovieSceneComponentTransformTrail::StartTracking()
 {
+	FMovieSceneTransformTrail::StartTracking();
+
 	UMovieSceneSection* Section = GetSection();
 	if (!Section)
 	{
@@ -399,74 +860,28 @@ bool FMovieSceneComponentTransformTrail::StartTracking()
 
 	if (IsAnythingSelected())
 	{
-		EditedTimes.Reset();
-		Section->Modify();
+		bStartTracking = true;
 		KeyTool->StartDragging();
 		return true;
 	}
 	return false;
 }
 
-static bool SetActorTransform(ISequencer* Sequencer, AActor *Actor, UMovieScene3DTransformSection* TransformSection, const TArray<FFrameNumber>& Frames,
-	const TArray<FTransform>& WorldTransformsToSnapTo)
+bool FMovieSceneComponentTransformTrail::ApplyDelta(const FVector& Pos, const FRotator& Rot, const FVector& WidgetLocation, bool bApplyToOffset)
 {
-	if (!Sequencer || !Sequencer->GetFocusedMovieSceneSequence())
+	if (bApplyToOffset || bIsOffsetMode)
 	{
-		return false;
+		return FMovieSceneTransformTrail::ApplyDelta(Pos, Rot, WidgetLocation, bApplyToOffset);
 	}
-	if (!TransformSection)
+	if (bStartTracking)
 	{
-		return false;
-	}
-
-	TransformSection->Modify();
-
-	TArray<FTransform> ParentWorldTransforms;
-	AActor* ParentActor = Actor->GetAttachParentActor();
-	if (ParentActor)
-	{
-		FActorForWorldTransforms ActorSelection;
-		ActorSelection.Actor = ParentActor;
-		ActorSelection.SocketName = Actor->GetAttachParentSocketName();
-		MovieSceneToolHelpers::GetActorWorldTransforms(Sequencer, ActorSelection, Frames, ParentWorldTransforms);
-	}
-	else
-	{
-		ParentWorldTransforms.SetNumUninitialized(Frames.Num());
-		for (FTransform& Transform : ParentWorldTransforms)
+		bStartTracking = false;
+		UMovieSceneSection* Section = GetSection();
+		if (!Section)
 		{
-			Transform = FTransform::Identity;
+			return false;
 		}
 	}
-
-	TArrayView<FMovieSceneDoubleChannel*> Channels = TransformSection->GetChannelProxy().GetChannels<FMovieSceneDoubleChannel>();
-	for (int32 Index = 0; Index < Frames.Num(); ++Index)
-	{
-		const FFrameNumber& Frame = Frames[Index];
-		FTransform ParentTransform = ParentWorldTransforms[Index];
-		FTransform WorldTransform = WorldTransformsToSnapTo[Index];
-		FTransform LocalTransform = WorldTransform.GetRelativeTransform(ParentTransform);
-		FVector Location = LocalTransform.GetLocation();
-		FRotator Rotation = LocalTransform.GetRotation().Rotator();
-			
-		TMovieSceneChannelData<FMovieSceneDoubleValue> ChannelData = Channels[0]->GetData();
-		MovieSceneToolHelpers::SetOrAddKey(ChannelData, Frame, Location.X);
-		ChannelData = Channels[1]->GetData();
-		MovieSceneToolHelpers::SetOrAddKey(ChannelData, Frame, Location.Y);
-		ChannelData = Channels[2]->GetData();
-		MovieSceneToolHelpers::SetOrAddKey(ChannelData, Frame, Location.Z);			
-	}
-
-	Channels[0]->AutoSetTangents();
-	Channels[1]->AutoSetTangents();
-	Channels[2]->AutoSetTangents();
-		
-	Sequencer->NotifyMovieSceneDataChanged(EMovieSceneDataChangeType::MovieSceneStructureItemsChanged);
-	return true;
-}
-
-bool FMovieSceneComponentTransformTrail::ApplyDelta(const FVector& Pos, const FRotator& Rot, const FVector& WidgetLocation)
-{
 	if (Component.IsValid() == false)
 	{
 		return false;
@@ -490,6 +905,7 @@ bool FMovieSceneComponentTransformTrail::ApplyDelta(const FVector& Pos, const FR
 	if (IsAnythingSelected())
 	{
 		TArray<FTransform> Transforms;
+		TArray<FTransform> ParentTransforms;
 		TArray<FFrameNumber> Times;
 		if (IsTrailSelected())
 		{
@@ -501,8 +917,10 @@ bool FMovieSceneComponentTransformTrail::ApplyDelta(const FVector& Pos, const FR
 				NewTransform.SetLocation(NewTransform.GetLocation() + Pos);
 				KeyInfo->Transform = NewTransform;
 				Transforms.Add(NewTransform);
+				ParentTransforms.Add(KeyInfo->ParentTransform);
 				Times.Add(KeyInfo->FrameNumber);
 			}
+			SelectedPos += Pos;
 		}
 		else
 		{
@@ -512,13 +930,11 @@ bool FMovieSceneComponentTransformTrail::ApplyDelta(const FVector& Pos, const FR
 				NewTransform.SetLocation(NewTransform.GetLocation() + Pos);
 				KeyInfo->Transform = NewTransform;
 				Transforms.Add(NewTransform);
+				ParentTransforms.Add(KeyInfo->ParentTransform);
 				Times.Add(KeyInfo->FrameNumber);
 			}
-			TArray<FFrameNumber> EditTimesSet;
-			EditedTimes = EditTimesSet;
 		}
-		AActor* Actor = Component.Get()->GetTypedOuter<AActor>();
-		SetActorTransform(GetSequencer().Get(), Actor, TransformSection, Times, Transforms);
+		UE::AIE::FSetTransformHelpers::SetActorTransform(GetSequencer().Get(), Component.Get(), TransformSection, Times, Transforms, ParentTransforms);
 		KeyTool->UpdateSelectedKeysTransform();
 		return true;
 	}
@@ -527,7 +943,8 @@ bool FMovieSceneComponentTransformTrail::ApplyDelta(const FVector& Pos, const FR
 
 bool FMovieSceneComponentTransformTrail::EndTracking()
 {
-	EditedTimes.Reset();
+	bStartTracking = false;
+	FMovieSceneTransformTrail::EndTracking();
 	UMovieSceneSection* Section = GetSection();
 	if (!Section)
 	{
@@ -544,7 +961,7 @@ bool FMovieSceneComponentTransformTrail::EndTracking()
 	return false;
 }
 
-bool FMovieSceneComponentTransformTrail::HandleAltClick( FEditorViewportClient* InViewportClient, HMotionTrailProxy* Proxy,FInputClick Click)
+bool FMovieSceneComponentTransformTrail::HandleAltClick( FEditorViewportClient* InViewportClient, HNewMotionTrailProxy* Proxy,FInputClick Click)
 {
 	if (Component.IsValid() == false)
 	{
@@ -565,122 +982,260 @@ bool FMovieSceneComponentTransformTrail::HandleAltClick( FEditorViewportClient*
 
 	TArray<FFrameNumber> Times;
 	TArray<FTransform> Transforms;
+	TArray<FTransform> ParentTransforms;
 	FFrameRate TickResolution = GetSequencer()->GetFocusedTickResolution();
 	FFrameRate DisplayResolution = GetSequencer()->GetFocusedDisplayRate();
 	//Snap to Frame. Assuming we want keys on frames.
-	FFrameTime GlobalTime = TickResolution.AsFrameTime(Proxy->Second);
+	FFrameTime GlobalTime = FFrameTime(Proxy->CurrentFrame,0.0);
 	FFrameTime DisplayTime = FFrameRate::TransformTime(GlobalTime, DisplayResolution, TickResolution);
 	DisplayTime = DisplayTime.RoundToFrame();
 	FFrameTime EvalTime = FFrameRate::TransformTime(DisplayTime, TickResolution, DisplayResolution);
 	Times.Add(EvalTime.RoundToFrame());
-	const FTransform Transform = GetTrajectoryTransforms()->GetInterp(Proxy->Second);
+	const FTransform Transform = ArrayOfTransforms->Interp(Proxy->CurrentFrame,CurrentFramesInfo->TransformIndices, CurrentFramesInfo->CurrentFrames);
+	const FTransform ParentTransform = ParentArrayOfTransforms->Interp(Proxy->CurrentFrame, CurrentFramesInfo->TransformIndices, CurrentFramesInfo->CurrentFrames);
 	Transforms.Add(Transform);
-	AActor* Actor = Component.Get()->GetTypedOuter<AActor>();
-	SetActorTransform(GetSequencer().Get(), Actor, TransformSection, Times, Transforms);
+	ParentTransforms.Add(ParentTransform);
+	UE::AIE::FSetTransformHelpers::SetActorTransform(GetSequencer().Get(), Component.Get(), TransformSection, Times, Transforms, ParentTransforms);
 	KeyTool->UpdateSelectedKeysTransform();
 	return true;
 }
 
+TUniquePtr<FMovieSceneTransformTrail::FMotionTrailState> FMovieSceneComponentTransformTrail::GetMotionTrailState() const
+{
+	TUniquePtr<FMovieSceneComponentTransformTrail::FMovieSceneComponentTransformTrailState> State = MakeUnique<FMovieSceneComponentTransformTrail::FMovieSceneComponentTransformTrailState>();
+	State->SaveFromTrail(this);
+	State->Component = Component;
+	return MoveTemp(State);
+}
+
+void FMovieSceneComponentTransformTrail::FMovieSceneComponentTransformTrailState::RestoreTrail(FTrailHierarchy* InTrailHierarchy)
+{
+	FSequencerTrailHierarchy* TrailHierarchy = static_cast<FSequencerTrailHierarchy*>(InTrailHierarchy);
+	USceneComponent* SceneComp = Component.Pin().Get();
+	if (!SceneComp)
+	{
+		return;
+	}
+	UMovieScene3DTransformTrack* Track = Cast< UMovieScene3DTransformTrack>(MainTrack.Get());
+	if (!Track)
+	{
+		return;
+	}
+	FGuid NewGuid = TrailHierarchy->AddComponentToHierarchy(BindingID, SceneComp, Track);
+	TrailHierarchy->PinTrail(NewGuid);
+	if (const TUniquePtr<FTrail>* Trail = TrailHierarchy->GetAllTrails().Find(NewGuid))
+	{
+		if (FMovieSceneTransformTrail* TransformTrail = static_cast<FMovieSceneTransformTrail*>(Trail->Get()))
+		{
+			SetToTrail(TransformTrail);
+		}
+	}
+
+}
+
 /****************************************************************************************************************
 *
-*   Control Rig Transform Trail
+*Movie Scene Socket Component Trail
 *
 *****************************************************************************************************************/
 
-FMovieSceneControlRigTransformTrail::FMovieSceneControlRigTransformTrail(USceneComponent* SceneComponent, const FLinearColor& InColor, const bool bInIsVisible, TWeakObjectPtr<UMovieSceneTrack> InWeakTrack, TSharedPtr<ISequencer> InSequencer,  const FName& InControlName)
-: FMovieSceneTransformTrail(SceneComponent, InColor, bInIsVisible, InWeakTrack, InSequencer)
-, ControlName(InControlName)
+void FMovieSceneSocketTransformTrail::UpdateNeedsEvaluation(const FGuid& InTrailGuid, FSequencerTrailHierarchy * SequencerHierarchy)
 {
+	if (Component.IsValid())
+	{
+		FMovieSceneTransformTrail::UpdateNeedsEvaluation(InTrailGuid, SequencerHierarchy);
+		FActorForWorldTransforms Actors;
+		Actors.Actor = Component->GetOwner();
+		Actors.Component = Component;
+		Actors.SocketName = SocketName;
+		SequencerHierarchy->EvaluateActor(InTrailGuid, Actors, ArrayOfTransforms, ParentArrayOfTransforms);
+	}
 }
 
 
-void FMovieSceneControlRigTransformTrail::UpdateCacheTimes(const FSceneContext& InSceneContext)
+FText FMovieSceneSocketTransformTrail::GetName() const
 {
-	/*
-	We don't do this anymore since we instead update the cache for EVERY drawn control rig trail otherwise
-	we evaluate the rig again for each trail.
-	*/
+	if (Component.IsValid())
+	{
+		FString ActorLabel = Component->GetOwner()->GetActorLabel(false);
+		FString SocketNameString = SocketName.ToString();
+
+		SocketNameString = FString::Printf(TEXT("%s:%s"), *ActorLabel,*SocketNameString);
+		return FText::FromString(SocketNameString);
+	}
+	return FText();
 }
 
-void FMovieSceneControlRigTransformTrail::SetUseKeysForTrajectory(bool bVal)
+bool FMovieSceneSocketTransformTrail::StartTracking()
 {
-	bUseKeysForTrajectory = bVal;
+
+	FMovieSceneTransformTrail::StartTracking();
+	if (IsAnythingSelected())
+	{
+		return true;
+	}
+	return false;
 }
 
-void FMovieSceneControlRigTransformTrail::GetTrajectoryPointsForDisplay(const FDisplayContext& InDisplayContext, TArray<FVector>& OutPoints, TArray<double>& OutSeconds)
+
+bool FMovieSceneSocketTransformTrail::ApplyDelta(const FVector& Pos, const FRotator& Rot, const FVector& WidgetLocation, bool bApplyToOffset)
 {
-	if (FTrajectoryDrawInfo* DI = GetDrawInfo())
+	if (bApplyToOffset || bIsOffsetMode)
 	{
-		if (bUseKeysForTrajectory == false)
-		{
-			DI->GetTrajectoryPointsForDisplay(InDisplayContext, OutPoints, OutSeconds);
-		}
-		else
+		return FMovieSceneTransformTrail::ApplyDelta(Pos, Rot, WidgetLocation, bApplyToOffset);
+	}
+	return false;
+}
+
+bool FMovieSceneSocketTransformTrail::EndTracking()
+{
+	FMovieSceneTransformTrail::EndTracking();
+
+	if (IsAnythingSelected())
+	{
+		//need to broadcast this so we can update other systems(including our own).
+		AActor* Actor = Component.Get()->GetTypedOuter<AActor>();
+		GEngine->BroadcastOnActorMoved(Actor);
+		return true;
+	}
+	return false;
+}
+
+
+void FMovieSceneSocketTransformTrail::DrawHUD(const FSceneView* View, FCanvas* Canvas)
+{
+	//todo for ticks
+}
+
+bool FMovieSceneSocketTransformTrail::HandleClick(const FGuid& Guid, FEditorViewportClient* InViewportClient, HHitProxy* InHitProxy, FInputClick Click)
+{
+	FMovieSceneTransformTrail::HandleClick(Guid, InViewportClient, InHitProxy, Click);
+	if (HNewMotionTrailProxy* HitProxy = HitProxyCast<HNewMotionTrailProxy>(InHitProxy))
+	{
+		if (HitProxy->Guid == Guid)
 		{
-			TRange<double> TimeRange = InDisplayContext.TimeRange;
 
-			TArray<FTrailKeyInfo*> AllKeys;
-			KeyTool->GetAllKeys(AllKeys);
-			TSharedPtr<ISequencer> SequencerPtr = WeakSequencer.Pin();
-			FFrameNumber CurrentFrame = SequencerPtr->GetLocalTime().Time.GetFrame();
-			FFrameRate TickResolution = SequencerPtr->GetFocusedTickResolution();
-			bool bCurrentHandled = false;
-			//get cached positions from all of the keys and the current frame if it's between two keys
-			for (int32 Index = 0; Index < AllKeys.Num(); ++Index)
+			if (Click.bShiftIsDown)
 			{
-				if (FTrailKeyInfo* KeyInfo = AllKeys[Index])
+				bIsSelected = true;
+				SelectedPos = HitProxy->Point;
+			}
+			else if (Click.bCtrlIsDown)
+			{
+				if (bIsSelected)
 				{
-					double LocalTime = TickResolution.AsSeconds(FFrameTime(KeyInfo->FrameNumber));
-					if (TimeRange.Contains(LocalTime))
-					{
-						FVector Point = DI->GetPoint(LocalTime);
-						OutPoints.Add(Point);
-						OutSeconds.Add(LocalTime);
-					}
-					if (LocalTime > TimeRange.GetUpperBoundValue())
-					{
-						break;
-					}
-					if (bCurrentHandled == false && ((Index + 1) < AllKeys.Num()))
-					{
-						if (CurrentFrame > KeyInfo->FrameNumber)
-						{
-							if (FTrailKeyInfo* NextKeyInfo = AllKeys[Index + 1])
-							{
-								if (CurrentFrame < NextKeyInfo->FrameNumber)
-								{
-									double CurrentTime = TickResolution.AsSeconds(FFrameTime(CurrentFrame));
-									if (TimeRange.Contains(CurrentTime))
-									{
-										FVector Point = DI->GetPoint(CurrentTime);
-										OutPoints.Add(Point);
-										OutSeconds.Add(CurrentTime);
-										bCurrentHandled = true;
-									}
-								}
-							}
-						}
-					}
+					bIsSelected = false;
+
+				}
+				else
+				{
+					bIsSelected = true;
+					SelectedPos = HitProxy->Point;
 				}
 			}
+			else
+			{
+				KeyTool->ClearSelection();
+				bIsSelected = true;
+				SelectedPos = HitProxy->Point;
+			}
+			return true;
 		}
 	}
+	bIsSelected = false;
+	return false;
+}
+
+
+bool FMovieSceneSocketTransformTrail::IsAnythingSelected(FVector& OutVectorPosition) const
+{
+	if (bIsSelected)
+	{
+		OutVectorPosition = SelectedPos;
+		return true;
+	}
+	return false;
 }
 
-void FMovieSceneControlRigTransformTrail::GetTickPointsForDisplay(const FDisplayContext& InDisplayContext, TArray<FVector2D>& OutTicks, TArray<FVector2D>& OutTickTangents)
+bool FMovieSceneSocketTransformTrail::IsAnythingSelected() const
 {
-	if (bUseKeysForTrajectory == false)
+	return (bIsSelected );
+}
+
+TUniquePtr<FMovieSceneTransformTrail::FMotionTrailState> FMovieSceneSocketTransformTrail::GetMotionTrailState() const
+{
+	TUniquePtr<FMovieSceneSocketTransformTrail::FMovieSceneSocketTransformTrailState> State = MakeUnique<FMovieSceneSocketTransformTrail::FMovieSceneSocketTransformTrailState>();
+	State->SaveFromTrail(this);
+	State->Component = Component;
+	State->SocketName = SocketName;
+	State->Color = GetDrawInfo()->GetColor();
+	State->PinnedStyle = GetDrawInfo()->GetStyle();
+	State->Color = GetDrawInfo()->GetColor();
+	State->PinnedStyle = GetDrawInfo()->GetStyle();
+	return MoveTemp(State);
+}
+
+void FMovieSceneSocketTransformTrail::FMovieSceneSocketTransformTrailState::RestoreTrail(FTrailHierarchy* InTrailHierarchy)
+{
+	FSequencerTrailHierarchy* TrailHierarchy = static_cast<FSequencerTrailHierarchy*>(InTrailHierarchy);
+	USceneComponent* SceneComp = Component.Pin().Get();
+	if (!SceneComp)
+	{
+		return;
+	}
+	FGuid NewGuid = TrailHierarchy->PinComponent(SceneComp, SocketName);
+	TrailHierarchy->PinTrail(NewGuid);
+	if (const TUniquePtr<FTrail>* Trail = TrailHierarchy->GetAllTrails().Find(NewGuid))
 	{
-		if (FTrajectoryDrawInfo* DI = GetDrawInfo())
+		if (FMovieSceneTransformTrail* TransformTrail = static_cast<FMovieSceneTransformTrail*>(Trail->Get()))
 		{
-			DI->GetTickPointsForDisplay(InDisplayContext, OutTicks, OutTickTangents);
+			SetToTrail(TransformTrail);
 		}
 	}
 }
 
+/****************************************************************************************************************
+*
+*   Control Rig Transform Trai
+*
+*****************************************************************************************************************/
+
+FMovieSceneControlRigTransformTrail::FMovieSceneControlRigTransformTrail(const FGuid& InBindingGuid, USceneComponent* SceneComponent,  const bool bInIsVisible, TWeakObjectPtr<UMovieSceneTrack> InWeakTrack, 
+	TSharedPtr<ISequencer> InSequencer,  const FName& InControlName, TSharedPtr<UE::AIE::FArrayOfTransforms>& InParentArrayOfTransforms)
+: FMovieSceneTransformTrail(InBindingGuid, SceneComponent, bInIsVisible, InWeakTrack, InSequencer)
+, ControlName(InControlName)
+{
+	ParentArrayOfTransforms = InParentArrayOfTransforms;
+}
+
+void FMovieSceneControlRigTransformTrail::UpdateNeedsEvaluation(const FGuid& InTrailGuid, FSequencerTrailHierarchy* SequencerHierarchy)
+{
+	UMovieSceneControlRigParameterSection* Section = Cast<UMovieSceneControlRigParameterSection>(GetSection());
+	if (!Section)
+	{
+		return;
+	}
+	UControlRig* ControlRig = Section->GetControlRig();
+	if (!ControlRig)
+	{
+		return;
+	}
+	FMovieSceneTransformTrail::UpdateNeedsEvaluation(InTrailGuid, SequencerHierarchy);
+
+	SequencerHierarchy->EvaluateControlRig(InTrailGuid, ControlRig, ControlName, ArrayOfTransforms);
+}
+
+FText FMovieSceneControlRigTransformTrail::GetName() const
+{
+	return FText::FromString(ControlName.ToString());
+}
+
+
 
 bool FMovieSceneControlRigTransformTrail::StartTracking()
 {
+	FMovieSceneTransformTrail::StartTracking();
+
 	UMovieSceneControlRigParameterSection* Section = Cast<UMovieSceneControlRigParameterSection>(GetSection());
 	if (!Section)
 	{
@@ -693,9 +1248,10 @@ bool FMovieSceneControlRigTransformTrail::StartTracking()
 	}
 	if(IsAnythingSelected())
 	{ 
-		EditedTimes.Reset();
-		Section->Modify();
-		ControlRig->Modify();
+		if ( bIsOffsetMode == false)
+		{
+			bStartTracking = true;
+		}
 		KeyTool->StartDragging();
 		return true;
 	}
@@ -713,8 +1269,29 @@ int32 FMovieSceneControlRigTransformTrail::GetChannelOffset() const
 	return INDEX_NONE;
 }
 
-bool FMovieSceneControlRigTransformTrail::ApplyDelta(const FVector& Pos, const FRotator& Rot, const FVector& WidgetLocation)
+bool FMovieSceneControlRigTransformTrail::ApplyDelta(const FVector& Pos, const FRotator& Rot, const FVector& WidgetLocation, bool bApplyToOffset)
 {
+	if (bApplyToOffset || bIsOffsetMode)
+	{
+		return FMovieSceneTransformTrail::ApplyDelta(Pos, Rot, WidgetLocation, bApplyToOffset);
+	}
+
+	if (bStartTracking)
+	{
+		bStartTracking = false;
+		UMovieSceneControlRigParameterSection* Section = Cast<UMovieSceneControlRigParameterSection>(GetSection());
+		if (!Section)
+		{
+			return false;
+		}
+		UControlRig* ControlRig = Section->GetControlRig();
+		if (!ControlRig)
+		{
+			return false;
+		}
+		Section->Modify();
+		ControlRig->Modify();
+	}
 	UMovieSceneControlRigParameterSection* Section = Cast<UMovieSceneControlRigParameterSection>(GetSection());
 	if (!Section)
 	{
@@ -733,6 +1310,12 @@ bool FMovieSceneControlRigTransformTrail::ApplyDelta(const FVector& Pos, const F
 
 	if (IsAnythingSelected())
 	{
+		FRigControlElement* ControlElement = ControlRig->GetHierarchy()->Find<FRigControlElement>((FRigElementKey(ControlName, ERigElementType::Control)));
+		if (ControlElement == nullptr)
+		{
+			return false;
+		}
+
 		FFrameRate TickResolution = GetSequencer()->GetFocusedTickResolution();
 		FRigControlModifiedContext Context;
 		Context.SetKey = EControlRigSetKey::Always;
@@ -741,7 +1324,7 @@ bool FMovieSceneControlRigTransformTrail::ApplyDelta(const FVector& Pos, const F
 		FMovieSceneSequenceTransform RootToLocalTransform = GetSequencer()->GetFocusedMovieSceneSequenceTransform();
 		FMovieSceneInverseSequenceTransform LocalToRootTransform = RootToLocalTransform.Inverse();
 
-		auto EvalControlRig = [&Context,&LocalToRootTransform, Pos, TickResolution, ControlRig,Player, this](FTrailKeyInfo* KeyInfo)
+		auto EvalControlRig = [&Context,&LocalToRootTransform, Pos, TickResolution, ControlRig,Player, ControlElement, this](FTrailKeyInfo* KeyInfo)
 		{
 			if (KeyInfo)
 			{
@@ -758,7 +1341,12 @@ bool FMovieSceneControlRigTransformTrail::ApplyDelta(const FVector& Pos, const F
 				NewTransform = NewTransform.GetRelativeTransform(KeyInfo->ParentTransform);
 				GetSequencer()->GetEvaluationTemplate().EvaluateSynchronousBlocking(MovieSceneContext);
 				ControlRig->Evaluate_AnyThread();
-				ControlRig->SetControlGlobalTransform(ControlName, NewTransform, true, Context, false /*undo*/, false /*bPrintPython*/, true/* bFixEulerFlips*/);
+				constexpr ERigTransformType::Type TransformType = ERigTransformType::CurrentGlobal;
+
+				FRigControlValue Value = ControlRig->GetControlValueFromGlobalTransform(ControlElement->GetKey().Name, NewTransform,TransformType);
+				NewTransform = Value.GetAsTransform(ControlElement->Settings.ControlType,ControlElement->Settings.PrimaryAxis);
+				FEulerTransform EulerTransform(NewTransform);
+				UE::AIE::FSetTransformHelpers::SetControlTransform(EulerTransform,ControlRig, ControlElement, Context);		
 			}
 		};
 
@@ -778,8 +1366,6 @@ bool FMovieSceneControlRigTransformTrail::ApplyDelta(const FVector& Pos, const F
 			{
 				EvalControlRig(KeyInfo);
 			}
-			TArray<FFrameNumber> EditTimesSet;
-			EditedTimes = EditTimesSet;
 		}
 		KeyTool->UpdateSelectedKeysTransform();
 		return true;
@@ -789,6 +1375,9 @@ bool FMovieSceneControlRigTransformTrail::ApplyDelta(const FVector& Pos, const F
 
 bool FMovieSceneControlRigTransformTrail::EndTracking()
 {
+	FMovieSceneTransformTrail::EndTracking();
+	bStartTracking = false;
+
 	UMovieSceneControlRigParameterSection* Section = Cast<UMovieSceneControlRigParameterSection>(GetSection());
 	if (!Section)
 	{
@@ -799,7 +1388,6 @@ bool FMovieSceneControlRigTransformTrail::EndTracking()
 	{
 		return false;
 	}
-	EditedTimes.Reset();
 
 	if (IsAnythingSelected())
 	{
@@ -818,7 +1406,7 @@ bool FMovieSceneControlRigTransformTrail::EndTracking()
 	return false;
 }
 
-bool FMovieSceneControlRigTransformTrail::HandleAltClick(FEditorViewportClient* InViewportClient, HMotionTrailProxy* Proxy, FInputClick Click)
+bool FMovieSceneControlRigTransformTrail::HandleAltClick(FEditorViewportClient* InViewportClient, HNewMotionTrailProxy* Proxy, FInputClick Click)
 {
 	UMovieSceneControlRigParameterSection* Section = Cast<UMovieSceneControlRigParameterSection>(GetSection());
 	if (!Section)
@@ -841,7 +1429,7 @@ bool FMovieSceneControlRigTransformTrail::HandleAltClick(FEditorViewportClient*
 	IMovieScenePlayer* Player = GetSequencer().Get();
 
 	//Snap to Frame. Assuming we want keys on frames.
-	FFrameTime GlobalTime = TickResolution.AsFrameTime(Proxy->Second);
+	FFrameTime GlobalTime = FFrameTime(Proxy->CurrentFrame,0.0);
 	FFrameTime DisplayTime = FFrameRate::TransformTime(GlobalTime, DisplayResolution, TickResolution);
 	DisplayTime = DisplayTime.RoundToFrame();
 	GlobalTime = FFrameRate::TransformTime(DisplayTime, TickResolution, DisplayResolution);
@@ -854,7 +1442,10 @@ bool FMovieSceneControlRigTransformTrail::HandleAltClick(FEditorViewportClient*
 	FMovieSceneContext MovieSceneContext = FMovieSceneContext(FMovieSceneEvaluationRange(GlobalTime, TickResolution), Player->GetPlaybackStatus()).SetHasJumped(true);
 	GetSequencer()->GetEvaluationTemplate().EvaluateSynchronousBlocking(MovieSceneContext);
 	ControlRig->Evaluate_AnyThread();
+	//todo handle constraints
 	FTransform NewTransform(ControlRig->GetControlGlobalTransform(ControlName));
+
+	//test constraints
 	ControlRig->SetControlGlobalTransform(ControlName, NewTransform, true, Context, false /*undo*/, false /*bPrintPython*/, true/* bFixEulerFlips*/);
 
 	//eval back at current time
@@ -878,9 +1469,45 @@ bool FMovieSceneControlRigTransformTrail::HandleAltClick(FEditorViewportClient*
 	return true;
 }
 
+TUniquePtr<FMovieSceneTransformTrail::FMotionTrailState> FMovieSceneControlRigTransformTrail::GetMotionTrailState() const
+{
+	TUniquePtr<FMovieSceneControlRigTransformTrail::FMovieSceneControlRigTransformTrailState> State = MakeUnique<FMovieSceneControlRigTransformTrail::FMovieSceneControlRigTransformTrailState>();
+	State->SaveFromTrail(this);
+	State->Owner = Cast<USkeletalMeshComponent>(Owner);
+	State->ControlName = ControlName;
+	return MoveTemp(State);
+}
+
+void FMovieSceneControlRigTransformTrail::FMovieSceneControlRigTransformTrailState::RestoreTrail(FTrailHierarchy* InTrailHierarchy)
+{
+	if (TSharedPtr<ISequencer> Sequencer = WeakSequencer.Pin())
+	{
+		FSequencerTrailHierarchy* TrailHierarchy = static_cast<FSequencerTrailHierarchy*>(InTrailHierarchy);
+		USkeletalMeshComponent* SkelMeshComp = Owner.Pin().Get();
+		if (!SkelMeshComp)
+		{
+			return;
+		}
+		UMovieSceneControlRigParameterTrack* CRTrack = Cast< UMovieSceneControlRigParameterTrack>(MainTrack.Get());
+		if (!CRTrack || !CRTrack->GetControlRig())
+		{
+			return;
+		}
+		FGuid NewGuid = TrailHierarchy->AddControlRigTrail(SkelMeshComp, CRTrack->GetControlRig(), CRTrack, ControlName);
+		TrailHierarchy->PinTrail(NewGuid);
+		if (const TUniquePtr<FTrail>* Trail = TrailHierarchy->GetAllTrails().Find(NewGuid))
+		{
+			if (FMovieSceneTransformTrail* TransformTrail = static_cast<FMovieSceneTransformTrail*>(Trail->Get()))
+			{
+				SetToTrail(TransformTrail);
+			}
+		}
+	}
+}
 
 
 } // namespace MovieScene
 } // namespace UE
 
 #undef LOCTEXT_NAMESPACE
+
diff --git a/Engine/Plugins/Editor/SequencerAnimTools/Source/SequencerAnimTools/Private/Sequencer/MovieSceneTransformTrail.h b/Engine/Plugins/Editor/SequencerAnimTools/Source/SequencerAnimTools/Private/Sequencer/MovieSceneTransformTrail.h
--- a/Engine/Plugins/Editor/SequencerAnimTools/Source/SequencerAnimTools/Private/Sequencer/MovieSceneTransformTrail.h
+++ b/Engine/Plugins/Editor/SequencerAnimTools/Source/SequencerAnimTools/Private/Sequencer/MovieSceneTransformTrail.h
@@ -7,6 +7,8 @@
 #include "MotionTrailMovieSceneKey.h"
 #include "MovieSceneTrack.h"
 #include "TrajectoryCache.h"
+#include "Tools/EvaluateSequencerTools.h"
+#include "Misc/Optional.h"
 
 class IMovieScenePlayer;
 
@@ -21,68 +23,169 @@ namespace UE
 namespace SequencerAnimTools
 {
 
+class FSequencerTrailHierarchy;
 
+//caching of tracks ans sections to handle signature changes.
+struct FCachedGuidsPerSignedObject
+{
+	TWeakObjectPtr<const UMovieSceneSignedObject> SignedObject;
+	FGuid CachedGuid;
+};
+
+struct FCachedTrackAndSections
+{
+	FCachedGuidsPerSignedObject Track;
+	TMap<TWeakObjectPtr<const UMovieSceneSection>, FCachedGuidsPerSignedObject> Sections;
+};
+
+struct FCachedTracksAndSections
+{
+	TMap<TWeakObjectPtr<UMovieSceneTrack>, FCachedTrackAndSections> Tracks;
+	//checks if cache is equal and recreates it
+	bool IsCacheEqual(const TArray <TWeakObjectPtr<UMovieSceneTrack>>& InTracks);
+	UMovieSceneSection* GetSection() const;
+	void CreateFrom(const TArray <TWeakObjectPtr<UMovieSceneTrack>>& InTracks);
+};
+
+//space that a trail may be in, includes Actor/ optional component and it's sequencer tracks/sections
+struct FOptionalParentSpace
+{
+	FOptionalParentSpace() : bIsValid(false) { ParentSpace.Actor.Actor = nullptr; ParentSpace.Actor.Component = nullptr; };
+
+	UE::AIE::FActorAndWorldTransforms ParentSpace; //this contains the actor/component
+	
+	FCachedTracksAndSections CachedTracksGuids;
+	bool bIsValid;
+
+	//set space and make it valid
+	void SetSpace(AActor* InActor, const FName& InComponentName);
+	void ClearSpace();
+};
+
+struct FDrawCacheData
+{
+	TArray<FVector> PointsToDraw;
+	TArray<FFrameNumber> Frames;
+	TArray<FLinearColor> Color;
+};
 class FMovieSceneTransformTrail : public FTrail
 {
 public:
-	FMovieSceneTransformTrail(USceneComponent* SceneComponent,const FLinearColor& InColor, const bool bInIsVisible, TWeakObjectPtr<UMovieSceneTrack> InWeakTrack, TSharedPtr<ISequencer> InSequencer);
+
+	FMovieSceneTransformTrail(const FGuid& InBindingID, USceneComponent* SceneComponent, const bool bInIsVisible, TWeakObjectPtr<UMovieSceneTrack> InWeakTrack, TSharedPtr<ISequencer> InSequencer);
+
+
 	// Begin FTrail interface
-	virtual ETrailCacheState UpdateTrail(const FSceneContext& InSceneContext) override;
-	virtual FTrajectoryCache* GetTrajectoryTransforms() override { return TrajectoryCache.Get(); }
-	virtual TRange<double> GetEffectiveRange() const override { return CachedEffectiveRange; }
+	virtual FTrailCurrentStatus UpdateTrail(const FNewSceneContext& NewSceneContext) override;
+	virtual void Interp(const FFrameNumber& Time, FTransform& OutTransform, FTransform& OutParentTransform) override;
+	virtual void UpdateFinished(const TRange<FFrameNumber>& UpdatedRange, const TArray<int32>& IndicesToCalcluate) override;
+	virtual void AddImportantTimes(TSet<FFrameNumber>& InOutImportantTimes) override;
+	virtual FTransform GetOffsetTransform() const override { return OffsetTransform; }
+	virtual void ClearOffsetTransform();
+	virtual void SetOffsetMode(); 
+
+	virtual FTransform GetParentSpaceTransform() const;
+	virtual void SetSpace(AActor* InActor, const FName& InComponentName) override { ParentSpace.SetSpace(InActor, InComponentName); }
+	virtual void ClearSpace() override { ParentSpace.ClearSpace(); 	ClearCachedData();}
+	virtual void ClearCachedData() override;
+	virtual void ReadyToDrawTrail(FColorState& ColorState, const FCurrentFramesInfo* InCurrentFramesInfo, bool bIsEvaluating, bool bIsPinned) override;
+	virtual void ActorChanged(AActor* InActor);
+	virtual void UpdateNeedsEvaluation(const FGuid& InTrailGuid, FSequencerTrailHierarchy* SequencerHierarchy);
+
 	virtual void Render(const FGuid& Guid, const FSceneView* View,  FPrimitiveDrawInterface* PDI) override;
+	virtual void RenderEvaluating(const FGuid& Guid, const FSceneView* View, FPrimitiveDrawInterface* PDI) override;
 	virtual void DrawHUD(const FSceneView* View, FCanvas* Canvas) override;
 
 	virtual bool HandleClick(const FGuid& Guid, FEditorViewportClient* InViewportClient, HHitProxy* HitProxy, FInputClick Click) override;
 	virtual bool IsAnythingSelected() const override;
 	virtual bool IsAnythingSelected(FVector& OutVectorPosition)const override;
+	virtual bool IsAnythingSelected(TArray<FVector>& OutVectorPositions)const override;
 	virtual bool IsTrailSelected() const override;
 	virtual bool StartTracking() override;
-	virtual bool ApplyDelta(const FVector& Pos, const FRotator& Rot, const FVector& WidgetLocation) override;
+	virtual bool ApplyDelta(const FVector& Pos, const FRotator& Rot, const FVector& WidgetLocation, bool bApplyToOffset) override;
 	virtual bool EndTracking() override;
+	virtual void HasStartedEvaluating() override;
 	virtual void TranslateSelectedKeys(bool bRight) override;
 	virtual void DeleteSelectedKeys() override;
 	virtual void SelectNone() override;
 	virtual int32 GetChannelOffset()const { return 0; }
-	virtual bool GetEditedTimes(const FTrailHierarchy* TrailHierarchy,const FFrameNumber& LastFrame, TArray<FFrameNumber>& OutEditedTimes);
-	virtual void UpdateKeysInRange(const TRange<double>& ViewRange) override;
 	virtual TArray<FFrameNumber> GetKeyTimes() const override;
 	virtual TArray<FFrameNumber> GetSelectedKeyTimes() const override;
 	virtual void ForceEvaluateNextTick() override;
+	virtual void GetColor(const FFrameNumber& CurrentTime, FColorState& InOutColorState) override;
 	// End FTrail interface
 
 	TSharedPtr<ISequencer> GetSequencer() const { return WeakSequencer.Pin(); }
 	FGuid GetCachedHierarchyGuid() const { return CachedHierarchyGuid; }
-	UMovieSceneTrack* GetTrack() const { return WeakTrack.Get(); }
 	UMovieSceneSection* GetSection() const;
 
+	void SetOffsetTransform(const FTransform& InOffsetTransform) { OffsetTransform = InOffsetTransform; ClearCachedData(); }
+	const UE::AIE::FSequencerTransformDependencies& GetTransformDependencies() { return TransformDependencies; }
+
 public:
 	TWeakPtr<ISequencer> WeakSequencer;
-	// Begin FMovieSceneTransformTrail interface
-	virtual void UpdateCacheTimes(const FSceneContext& InSceneContext) = 0;
+
 protected:
-	virtual bool HandleAltClick(FEditorViewportClient* InViewportClient, HMotionTrailProxy* Proxy, FInputClick Click);
-	// End FMovieSceneTransformTrail interface
+	virtual bool HandleAltClick(FEditorViewportClient* InViewportClient, HNewMotionTrailProxy* Proxy, FInputClick Click);
 
-	// when editing interactively the tool may specify just a specify range to recalculate
-	void CalculateEditedTimes(const FTrailHierarchy* TrailHierarchy, const FFrameNumber& LastFrame);
-	void ClearEditedTiimes() { EditedTimes.Reset(); }
+	//internal draw trail
+	void InternalDrawTrail(const FGuid& Guid, const FSceneView* View, FPrimitiveDrawInterface* PDI, const FDrawCacheData& InDrawData, bool bHitTesting) const;
 
-	IMovieScenePlayer* MovieScenePlayer;
-	TRange<double> GetEffectiveSectionRange(int32 ChannelOffset) const;
-	TRange<double> CachedEffectiveRange;
+	//if selected or selected and offset mode we draw different colors
+	TOptional<FLinearColor> GetOverrideColor() const;
 
+protected:
+	//some interactive states/modes
 	bool bIsSelected;
+	bool bIsOffsetMode;
 	FVector SelectedPos;
+	bool bIsTracking; 
 
+	//main tool for adjusting keys
 	TUniquePtr<FMotionTraiMovieScenelKeyTool> KeyTool;
-	TUniquePtr<FArrayTrajectoryCache> TrajectoryCache;
 
-	FGuid LastTransformSectionSig;
+	//previous parent space
+	mutable FTransform PreviousParentSpaceTM = FTransform::Identity;
+	//cached guid in the hierarchy
 	FGuid CachedHierarchyGuid;
-	TWeakObjectPtr<UMovieSceneTrack> WeakTrack;
+protected:
+	bool TrailHasChanged();
+
+protected: 
+	FGuid BindingID;
+	TWeakObjectPtr<UMovieSceneTrack> MainTrack;
+	UE::AIE::FSequencerTransformDependencies TransformDependencies;
+
+	FTransform OffsetTransform;
+	TSharedPtr<UE::AIE::FArrayOfTransforms> ArrayOfTransforms;
+	TSharedPtr<UE::AIE::FArrayOfTransforms> ParentArrayOfTransforms;
+	const FCurrentFramesInfo* CurrentFramesInfo;
+
+	//optional parent space
+	FOptionalParentSpace ParentSpace;
+	//cached drawing data
+	mutable FDrawCacheData CachedDrawData;
+	//old cached draw data used to draw the trail when we are updating it
+	FDrawCacheData PreviousCachedDrawData;
+	//cached color values for heat map
+	TArray<FLinearColor> CachedHeatMap;
 
-	TOptional<TArray<FFrameNumber>> EditedTimes;
+
+	class FMovieSceneTransformTrailState : public FTrail::FMotionTrailState
+	{
+	public:
+		FGuid BindingID;
+		TWeakObjectPtr<UMovieSceneTrack> MainTrack;
+		TWeakPtr<ISequencer> WeakSequencer;
+		FLinearColor Color;
+		EMotionTrailTrailStyle PinnedStyle;
+		FTransform OffsetTransform;
+		FOptionalParentSpace ParentSpace;
+		void SaveFromTrail(const FMovieSceneTransformTrail* InTrail);
+		void SetToTrail(FMovieSceneTransformTrail* InTrail) const;
+		 
+	};
+	friend class FMovieSceneTransformTrailState;
 };
 
 
@@ -90,51 +193,121 @@ class FMovieSceneComponentTransformTrail : public FMovieSceneTransformTrail
 {
 public:
 
-	FMovieSceneComponentTransformTrail(USceneComponent* InComponent, const FLinearColor& InColor, const bool bInIsVisible, TWeakObjectPtr<UMovieSceneTrack> InWeakTrack, TSharedPtr<ISequencer> InSequencer)
-		: FMovieSceneTransformTrail(InComponent, InColor, bInIsVisible, InWeakTrack, InSequencer),
+	FMovieSceneComponentTransformTrail(const FGuid& InBindingID, USceneComponent* InComponent, const bool bInIsVisible, TWeakObjectPtr<UMovieSceneTrack> InWeakTrack, TSharedPtr<ISequencer> InSequencer)
+		: FMovieSceneTransformTrail(InBindingID, InComponent, bInIsVisible, InWeakTrack, InSequencer),
 		Component(InComponent)
-		
 	{
 	}
 
+	//new
+	virtual void UpdateNeedsEvaluation(const FGuid& InTrailGuid, FSequencerTrailHierarchy* SequencerHierarchy) override;
+	virtual FText GetName() const override;
+	virtual TUniquePtr<FMotionTrailState> GetMotionTrailState() const;
+
 
 private:
 	TWeakObjectPtr<USceneComponent> Component;
 	// Begin FMovieSceneTransformTrail interface
-	virtual void UpdateCacheTimes(const FSceneContext& InSceneContext) override;
 	virtual bool StartTracking() override;
-	virtual bool ApplyDelta(const FVector& Pos, const FRotator& Rot, const FVector& WidgetLocation) override;
+	virtual bool ApplyDelta(const FVector& Pos, const FRotator& Rot, const FVector& WidgetLocation, bool bApplyToOffset) override;
 	virtual bool EndTracking() override;
 
 protected:
-	virtual bool HandleAltClick(FEditorViewportClient* InViewportClient, HMotionTrailProxy* Proxy, FInputClick Click) override;
+	virtual bool HandleAltClick(FEditorViewportClient* InViewportClient, HNewMotionTrailProxy* Proxy, FInputClick Click) override;
 
 	// End FMovieSceneTransformTrail interface
+	class FMovieSceneComponentTransformTrailState : public FMovieSceneTransformTrailState
+	{
+	public:
+		TWeakObjectPtr<USceneComponent> Component;
+		virtual void RestoreTrail(FTrailHierarchy* TrailHierarchy) override;
+	};
+private:
+	bool bStartTracking = false;
 
+	friend class FMovieSceneComponentTransformTrailState;
 };
 
+class FMovieSceneSocketTransformTrail : public FMovieSceneTransformTrail
+{
+public:
+
+	FMovieSceneSocketTransformTrail(const FGuid& InBindingID, USceneComponent* InComponent, const FName& InSocketName,
+		 const bool bInIsVisible, TWeakObjectPtr<UMovieSceneTrack> InWeakTrack,
+		TSharedPtr<ISequencer> InSequencer)
+		: FMovieSceneTransformTrail(InBindingID, InComponent, bInIsVisible, InWeakTrack, InSequencer),
+		Component(InComponent),  SocketName(InSocketName)
+	{
+	}
+
+	//new
+	virtual void UpdateNeedsEvaluation(const FGuid& InTrailGuid, FSequencerTrailHierarchy* SequencerHierarchy) override;
+	virtual FText GetName() const override;
+	virtual TUniquePtr<FMotionTrailState> GetMotionTrailState() const;
+
+
+private:
+	// Begin FMovieSceneTransformTrail interface
+	virtual bool StartTracking() override;
+	virtual bool ApplyDelta(const FVector& Pos, const FRotator& Rot, const FVector& WidgetLocation, bool bApplyToOffset) override;
+	virtual bool EndTracking() override;
+	virtual void DrawHUD(const FSceneView* View, FCanvas* Canvas) override;
+
+	virtual bool HandleClick(const FGuid& Guid, FEditorViewportClient* InViewportClient, HHitProxy* HitProxy, FInputClick Click) override;
+	virtual bool IsAnythingSelected() const override;
+	virtual bool IsAnythingSelected(FVector& OutVectorPosition)const override;
+
+protected:
+
+	// End FMovieSceneTransformTrail interface
+	TWeakObjectPtr<USceneComponent> Component;
+	FName SocketName;
+
+protected:
+	class FMovieSceneSocketTransformTrailState : public FMovieSceneTransformTrailState
+	{
+	public:
+		TWeakObjectPtr<USceneComponent> Component;
+		FName SocketName;
+		virtual void RestoreTrail(FTrailHierarchy* TrailHierarchy) override;
+	};
+};
 
 class FMovieSceneControlRigTransformTrail : public FMovieSceneTransformTrail
 {
 public:
 
-	FMovieSceneControlRigTransformTrail(USceneComponent* SceneComponent, const FLinearColor& InColor, const bool bInIsVisible, TWeakObjectPtr<UMovieSceneTrack> InWeakTrack, TSharedPtr<ISequencer> InSequencer, const FName& InControlName);
+	FMovieSceneControlRigTransformTrail(const FGuid& InBindingID, USceneComponent* SceneComponent,  const bool bInIsVisible, TWeakObjectPtr<UMovieSceneTrack> InWeakTrack,
+		TSharedPtr<ISequencer> InSequencer, const FName& InControlName, TSharedPtr<UE::AIE::FArrayOfTransforms>& InParentArrayOfTransforms);
+
+	virtual void UpdateNeedsEvaluation(const FGuid& InTrailGuid, FSequencerTrailHierarchy* SequencerHierarchy) override;
+	virtual FText GetName() const override;
+	virtual TUniquePtr<FMotionTrailState> GetMotionTrailState() const;
 
-	virtual void UpdateCacheTimes(const FSceneContext& InSceneContext) override;
 	virtual bool StartTracking() override;
-	virtual bool ApplyDelta(const FVector& Pos, const FRotator& Rot, const FVector& WidgetLocation) override;
+	virtual bool ApplyDelta(const FVector& Pos, const FRotator& Rot, const FVector& WidgetLocation, bool bApplyToOffset) override;
 	virtual bool EndTracking() override;
 	virtual int32 GetChannelOffset() const override;
-	virtual void GetTrajectoryPointsForDisplay(const FDisplayContext& InDisplayContext, TArray<FVector>& OutPoints, TArray<double>& OutSeconds) override;
-	virtual void GetTickPointsForDisplay(const FDisplayContext& InDisplayContext, TArray<FVector2D>& OutTicks, TArray<FVector2D>& OutTickTangents) override;
 
-	void SetUseKeysForTrajectory(bool bVal);
 protected:
-	virtual bool HandleAltClick(FEditorViewportClient* InViewportClient, HMotionTrailProxy* Proxy, FInputClick Click) override;
+	virtual bool HandleAltClick(FEditorViewportClient* InViewportClient, HNewMotionTrailProxy* Proxy, FInputClick Click) override;
 
 private:
 	bool bUseKeysForTrajectory = false; //on when interatively moving
+	bool bStartTracking = false;
+
 	FName ControlName;
+
+private:
+	class FMovieSceneControlRigTransformTrailState : public FMovieSceneTransformTrailState
+	{
+	public:
+		TWeakObjectPtr<USkeletalMeshComponent> Owner;
+		FName ControlName;
+		virtual void RestoreTrail(FTrailHierarchy* TrailHierarchy) override;
+
+	};
+
 };
 
 } // namespace MovieScene
