Anim In Engine: Motion trails 2.0.

diff --git a/Engine/Plugins/Editor/SequencerAnimTools/Source/SequencerAnimTools/Private/Sequencer/SequencerTrailHierarchy.cpp b/Engine/Plugins/Editor/SequencerAnimTools/Source/SequencerAnimTools/Private/Sequencer/SequencerTrailHierarchy.cpp
--- a/Engine/Plugins/Editor/SequencerAnimTools/Source/SequencerAnimTools/Private/Sequencer/SequencerTrailHierarchy.cpp
+++ b/Engine/Plugins/Editor/SequencerAnimTools/Source/SequencerAnimTools/Private/Sequencer/SequencerTrailHierarchy.cpp
@@ -5,11 +5,10 @@
 #include "Editor.h"
 #include "Tracks/MovieScene3DTransformTrack.h"
 #include "Evaluation/MovieSceneEvaluationTemplateInstance.h"
-#include "Tracks/MovieSceneSkeletalAnimationTrack.h"
-
+#include "LevelEditorViewport.h"
 #include "Evaluation/MovieScenePlayback.h"
+#include "Tracks/MovieSceneSkeletalAnimationTrack.h"
 #include "Sequencer/MovieSceneControlRigParameterTrack.h"
-#include "AnimationBoneTrail.h"
 #include "LevelSequence.h"
 #include "Tools/MotionTrailOptions.h"
 #include "LevelSequenceCameraSettings.h"
@@ -20,12 +19,61 @@
 #include "MovieSceneToolHelpers.h"
 #include "Framework/Application/SlateApplication.h"
 #include "Sequencer/MotionTrailMovieSceneKey.h"
+#include "ControlRigObjectBinding.h"
+#include "Components/SkeletalMeshComponent.h"
+#include "Engine/SkeletalMesh.h"
+#include "Selection.h"
+#include "Editor/UnrealEdEngine.h"
 
 namespace UE
 {
 namespace SequencerAnimTools
 {
 
+void FColorState::Setup(FTrailHierarchy* TrailHierarchy)
+{
+	Options = UMotionTrailToolOptions::GetTrailOptions();
+	if(Options)
+	{
+		FSequencerTrailHierarchy* SequencerTrailHierarchy = static_cast<FSequencerTrailHierarchy*>(TrailHierarchy);
+		SequencerTime = SequencerTrailHierarchy->GetLocalTime();
+		TicksPerFrame = SequencerTrailHierarchy->GetFramesPerFrame();
+	}
+}
+
+EMotionTrailTrailStyle FColorState::GetStyle() const
+{
+	if (bIsPinned)
+	{
+		return PinnedStyle;
+	}
+	if (Options)
+	{
+		return Options->TrailStyle;
+	}
+	return PinnedStyle;
+}
+
+void FColorState::ReadyForTrail(bool bInIsPinned, EMotionTrailTrailStyle InPinnedStyle)
+{
+	bFirstFrame = true;
+	bIsPinned = bInIsPinned;
+	PinnedStyle = InPinnedStyle;
+	CalculatedColor = FColor(0xffffff);
+}
+
+FFrameNumber FSequencerTrailHierarchy::GetLocalTime() const
+{
+	if(TSharedPtr<ISequencer> Sequencer = WeakSequencer.Pin())
+	{
+		return Sequencer->GetLocalTime().Time.GetFrame();
+	}
+	return FFrameNumber(0);
+}
+
+	
+using namespace UE::AIE;
+TArray<TUniquePtr <FTrail::FMotionTrailState >> FSequencerTrailHierarchy::PreviouslyPinnedTrails;
 
 void FSequencerTrailHierarchy::Initialize()
 {
@@ -75,6 +123,51 @@ void FSequencerTrailHierarchy::Initialize()
 				GuidTrailPair.Value->ForceEvaluateNextTick();
 			}
 		}
+		if (PropertyName == GET_MEMBER_NAME_CHECKED(UMotionTrailToolOptions, TrailStyle) ||
+			PropertyName == GET_MEMBER_NAME_CHECKED(UMotionTrailToolOptions, DefaultColor) ||
+			PropertyName == GET_MEMBER_NAME_CHECKED(UMotionTrailToolOptions, TimePreColor) ||
+			PropertyName == GET_MEMBER_NAME_CHECKED(UMotionTrailToolOptions, TimePostColor) ||
+			PropertyName == GET_MEMBER_NAME_CHECKED(UMotionTrailToolOptions, DashPreColor) ||
+			PropertyName == GET_MEMBER_NAME_CHECKED(UMotionTrailToolOptions, DashPostColor)
+			)
+		{
+			for (const TPair<FGuid, TUniquePtr<FTrail>>& GuidTrailPair : GetAllTrails())
+			{
+				GuidTrailPair.Value->ClearCachedData();
+			}
+		}
+		else if (PropertyName == GET_MEMBER_NAME_CHECKED(UMotionTrailToolOptions, bShowSelectedTrails))
+		{
+			if (UMotionTrailToolOptions::GetTrailOptions()->bShowSelectedTrails == false)
+			{
+				TArray<FGuid> TrailsToRemove; //if trail is not visible we just get rid of it
+				// <parent, current>
+				for (const TPair<FGuid, TUniquePtr<FTrail>>& GuidTrailPair : GetAllTrails())
+				{
+					if(GetVisibilityManager().IsTrailAlwaysVisible(GuidTrailPair.Key) == false)
+					{
+						TrailsToRemove.Add(GuidTrailPair.Key);
+					}
+				}	
+				for (const FGuid& Key : TrailsToRemove)
+				{
+					RemoveTrail(Key);
+				}
+			}
+			else
+			{
+				TSharedPtr<ISequencer> Sequencer = WeakSequencer.Pin();
+				if (Sequencer)
+				{
+					TArray<FGuid> SequencerSelectedObjects;
+					Sequencer->GetSelectedObjects(SequencerSelectedObjects);
+					UpdateSequencerBindings(SequencerSelectedObjects,
+						[this](UObject* Object, FTrail*, FGuid Guid) {
+							VisibilityManager.Selected.Add(Guid);
+						});
+				}
+			}
+		}
 	});
 
 
@@ -84,6 +177,20 @@ void FSequencerTrailHierarchy::Initialize()
 	GEngine->OnActorsMoved().AddRaw(this, &FSequencerTrailHierarchy::OnActorsChangedSomehow);
 	//GEditor->RegisterForUndo(this);
 
+	RegisterMotionTrailOptionDelegates();
+	
+	if (PreviouslyPinnedTrails.Num() > 0)
+	{
+		for (TUniquePtr <FTrail::FMotionTrailState >& PinnedState : PreviouslyPinnedTrails)
+		{
+			if (PinnedState.IsValid())
+			{
+				PinnedState->RestoreTrail(this);
+			}
+		}
+		PreviouslyPinnedTrails.Reset();
+	}
+	
 }
 
 void FSequencerTrailHierarchy::Destroy()
@@ -107,9 +214,26 @@ void FSequencerTrailHierarchy::Destroy()
 		}
 	}
 
+	PreviouslyPinnedTrails.Reset();
+	if (GetAllTrails().Num() > 0)
+	{
+		
+		for (const TPair<FGuid, TUniquePtr<FTrail>>& GuidTrailPair : GetAllTrails())
+		{
+			if (VisibilityManager.IsTrailAlwaysVisible(GuidTrailPair.Key)) //if not always visible then it's selected
+			{
+				TUniquePtr<FTrail::FMotionTrailState> State = GuidTrailPair.Value->GetMotionTrailState();
+				if (State.IsValid())
+				{
+					PreviouslyPinnedTrails.Add(MoveTemp(State));
+				}
+			}
+		}
+	}
+
 	ObjectsTracked.Reset();
-	BonesTracked.Reset();
 	ControlsTracked.Reset();
+	SocketsTracked.Reset();
 	AllTrails.Reset();
 
 	GEngine->OnLevelActorAdded().RemoveAll(this);
@@ -117,24 +241,18 @@ void FSequencerTrailHierarchy::Destroy()
 	GEngine->OnActorMoved().RemoveAll(this);
 	GEngine->OnActorsMoved().RemoveAll(this);
 	//GEditor->UnregisterForUndo(this);
+	VisibilityManager.Reset();
+	UnRegisterMotionTrailOptionDelegates();
+
 
+	UMotionTrailToolOptions::GetTrailOptions()->ResetPinnedItems();
 }
 
 void FSequencerTrailHierarchy::OnActorChangedSomehow(AActor* InActor)
 {
 	for (const TPair<FGuid, TUniquePtr<FTrail>>& GuidTrailPair : GetAllTrails())
 	{
-		USceneComponent* Component = Cast<USceneComponent>(GuidTrailPair.Value->GetOwner());
-		if (Component)
-		{
-			TArray<const UObject*> ParentActors;
-			UObject* Parent = Component->GetOwner();
-			MovieSceneToolHelpers::GetParents(ParentActors, Parent);
-			if (ParentActors.Contains(InActor))
-			{
-				GuidTrailPair.Value->ForceEvaluateNextTick();
-			}
-		}
+		GuidTrailPair.Value->ActorChanged(InActor);
 	}
 }
 
@@ -144,25 +262,26 @@ void FSequencerTrailHierarchy::OnActorsChangedSomehow(TArray<AActor*>& InActors)
 	{
 		for (AActor* InActor : InActors)
 		{
-			USceneComponent* Component = Cast<USceneComponent>(GuidTrailPair.Value->GetOwner());
-			if (Component)
-			{
-				TArray<const UObject*> ParentActors;
-				UObject* Parent = Component->GetOwner();
-				MovieSceneToolHelpers::GetParents(ParentActors, Parent);
-				if (ParentActors.Contains(InActor))
-				{
-					GuidTrailPair.Value->ForceEvaluateNextTick();
-				}
-			}
+			GuidTrailPair.Value->ActorChanged(InActor);
 		}
 	}
 }
 
-double FSequencerTrailHierarchy::GetSecondsPerSegment() const 
-{ 
-	double SecondsPerFrame =  WeakSequencer.Pin()->GetFocusedDisplayRate().AsInterval();
-	return SecondsPerFrame; // no longer use this to avoid divide by zero options / double(UMotionTrailToolOptions::GetTrailOptions()->EvalsPerFrame);
+void FSequencerTrailHierarchy::CalculateEvalRangeArray()
+{
+	TicksPerSegment = GetFramesPerFrame();
+	CurrentFramesInfo.SetViewRange(TickViewRange);
+	if (LastTicksPerSegment != TicksPerSegment || TickEvalRange != LastTickEvalRange)
+	{	
+		LastTicksPerSegment = TicksPerSegment;
+		LastTickEvalRange = TickEvalRange;
+		for (const TPair<FGuid, TUniquePtr<FTrail>>& GuidTrailPair : GetAllTrails())
+		{
+			GuidTrailPair.Value->ForceEvaluateNextTick();
+		}
+		CurrentFramesInfo.SetUpFrameTimes(TickEvalRange, 
+			TicksPerSegment);
+	}
 }
 
 FFrameNumber FSequencerTrailHierarchy::GetFramesPerSegment() const
@@ -172,6 +291,11 @@ FFrameNumber FSequencerTrailHierarchy::GetFramesPerSegment() const
 	return FramesPerTick;
 }
 
+const FCurrentFramesInfo* FSequencerTrailHierarchy::GetCurrentFramesInfo() const
+{
+	return &CurrentFramesInfo;
+}
+
 FFrameNumber FSequencerTrailHierarchy::GetFramesPerFrame() const
 {
 	FFrameNumber FramesPerTick = FFrameRate::TransformTime(FFrameNumber(1), WeakSequencer.Pin()->GetFocusedDisplayRate(),
@@ -189,7 +313,7 @@ void FSequencerTrailHierarchy::RemoveTrail(const FGuid& Key)
 	}
 	else
 	{
-		for (TPair<USkeletalMeshComponent*, TMap<FName, FGuid>>& CompMapPair : BonesTracked)
+		for (TPair<USceneComponent*, TMap<FName, FGuid>>& CompMapPair : SocketsTracked)
 		{
 			if (const FName* FoundBone = CompMapPair.Value.FindKey(Key))
 			{
@@ -197,16 +321,36 @@ void FSequencerTrailHierarchy::RemoveTrail(const FGuid& Key)
 				return;
 			}
 		}
-		for (TPair<UControlRig*, TMap<FName, FGuid>>& CompMapPair : ControlsTracked)
+		for (TPair<UControlRig*, FControlMapAndTransforms>& CompMapPair : ControlsTracked)
 		{
-			if (const FName* FoundControl = CompMapPair.Value.FindKey(Key))
+			if (const FName* FoundControl = CompMapPair.Value.NameToTrail.FindKey(Key))
 			{
-				CompMapPair.Value.Remove(*FoundControl);
+				CompMapPair.Value.NameToTrail.Remove(*FoundControl);
 				return;
 			}
 		}
 	}
 }
+
+bool FSequencerTrailHierarchy::IsTrailEvaluating(const FGuid& InGuid, bool bIndirectOnly) const
+{
+	if (EvaluatingTrails.Contains(InGuid))
+	{
+		if (bIndirectOnly)
+		{
+			if (const TUniquePtr<FTrail>* Trail = GetAllTrails().Find(InGuid))
+			{
+				if (Trail->Get()->IsAnythingSelected())
+				{
+					return false;
+				}
+			}
+		}
+		return true;
+	}
+	return false;
+}
+
 struct FTrailControlTransforms
 {
 	FName ControlName;
@@ -214,223 +358,149 @@ struct FTrailControlTransforms
 	TArray<FTransform> Transforms;
 };
 
-void FSequencerTrailHierarchy::UpdateControlRig(const TArray<FFrameNumber> &Frames,UControlRig* ControlRig, TMap<FName, FGuid >& CompMapPair, bool bUseEditedTimes)
+void FSequencerTrailHierarchy::EvaluateSequencerAndSetTransforms()
 {
-	if (TSharedPtr<IControlRigObjectBinding> ObjectBinding = ControlRig->GetObjectBinding())
+	if (EvaluatingActors.Num() > 0 || EvaluatingControlRigs.Num() > 0)
 	{
-		TSharedPtr<ISequencer> Sequencer = WeakSequencer.Pin();
-		FFrameRate TickResolution = Sequencer->GetFocusedTickResolution();
-
-		UWorld* World = GEditor->GetEditorWorldContext().World();
-		FMovieSceneSequencePlaybackSettings Settings;
-		FLevelSequenceCameraSettings CameraSettings;
-		FMovieSceneSequenceIDRef Template = MovieSceneSequenceID::Root;
-		FMovieSceneInverseSequenceTransform LocalToRootTransform = Sequencer->GetFocusedMovieSceneSequenceTransform().Inverse();
-		IMovieScenePlayer* Player = Sequencer.Get();
-		ULevelSequence* LevelSequence = Cast<ULevelSequence>(Sequencer->GetFocusedMovieSceneSequence());
-
-		USceneComponent* Component = Cast<USceneComponent>(ObjectBinding->GetBoundObject());
-		if (Component && Player)
+		if (TSharedPtr<ISequencer> SequencerPtr = WeakSequencer.Pin())
 		{
-			AActor* Actor = Component->GetTypedOuter< AActor >();
-			if (Actor)
+			TSet<FFrameNumber> MustHaveTimes;
+			TSet<const UMovieSceneTrack*> DependentTracks;
+			MustHaveTimes.Add(TickEvalRange.GetLowerBoundValue().Value);
+			MustHaveTimes.Add(TickEvalRange.GetUpperBoundValue().Value);
+			FFrameNumber CurrentFrame = SequencerPtr->GetLocalTime().Time.GetFrame();
+
+			MustHaveTimes.Add(CurrentFrame);
+			if (TickViewRange.GetLowerBoundValue().Value != TickEvalRange.GetLowerBoundValue().Value)
 			{
-				TArray<FTransform> ControlRigParentWorldTransforms;
-				FActorForWorldTransforms ControlRigActorSelection;
-				ControlRigActorSelection.Actor = Actor;
-
-				MovieSceneToolHelpers::GetActorWorldTransforms(Player, LevelSequence, Template, ControlRigActorSelection, Frames, ControlRigParentWorldTransforms);
-
-				TArray<FTrailControlTransforms> TrailControlTransforms;
-				TrailControlTransforms.SetNum(CompMapPair.Num());
-				int32 PairIndex = 0;
-				for (TPair<FName, FGuid >& Pair : CompMapPair)
-				{
-					TrailControlTransforms[PairIndex].ElementGuid = Pair.Value;
-					TrailControlTransforms[PairIndex].ControlName = Pair.Key;
-					TrailControlTransforms[PairIndex].Transforms.SetNum(Frames.Num());
-					++PairIndex;
-				}
-
-				for (int32 Index = 0; Index < Frames.Num(); ++Index)
+				MustHaveTimes.Add(TickViewRange.GetLowerBoundValue().Value);
+			}
+			if (TickViewRange.GetUpperBoundValue().Value != TickEvalRange.GetUpperBoundValue().Value)
+			{
+				MustHaveTimes.Add(TickViewRange.GetUpperBoundValue().Value);
+			}
+			for (const TPair<FGuid, TUniquePtr<FTrail>>& GuidTrailPair : GetAllTrails())
+			{
+				if (EvaluatingTrails.Contains(GuidTrailPair.Key))
 				{
-					const FFrameNumber& FrameNumber = Frames[Index];
-					FFrameTime GlobalTime(FrameNumber);
-					GlobalTime = LocalToRootTransform.TryTransformTime(GlobalTime).Get(GlobalTime);
-
-					FMovieSceneContext Context = FMovieSceneContext(FMovieSceneEvaluationRange(GlobalTime, TickResolution), Player->GetPlaybackStatus()).SetHasJumped(true);
-
-					Player->GetEvaluationTemplate().EvaluateSynchronousBlocking(Context);
-					ControlRig->Evaluate_AnyThread();
-					for (FTrailControlTransforms& TrailControlTransform : TrailControlTransforms)
+					GuidTrailPair.Value->AddImportantTimes(MustHaveTimes);
+					const UE::AIE::FSequencerTransformDependencies& TransformDependencies = (static_cast<FMovieSceneTransformTrail*>(GuidTrailPair.Value.Get()))->GetTransformDependencies();
+					for (const TPair<TWeakObjectPtr<UMovieSceneTrack>,FGuid>& Track : TransformDependencies.Tracks)
 					{
-						TrailControlTransform.Transforms[Index] = ControlRig->GetControlGlobalTransform(TrailControlTransform.ControlName) * ControlRigParentWorldTransforms[Index];
-						if (bUseEditedTimes)
+						if (Track.Key.IsValid())
 						{
-							double Sec = TickResolution.AsSeconds(FFrameTime(FrameNumber));
-							AllTrails[TrailControlTransform.ElementGuid]->GetTrajectoryTransforms()->Set(Sec, TrailControlTransform.Transforms[Index]);
+							DependentTracks.Add(Track.Key.Get());
 						}
 					}
 				}
-				for (FTrailControlTransforms& TrailControlTransform : TrailControlTransforms)
+			}
+			CurrentFramesInfo.AddMustHaveTimes(MustHaveTimes);
+
+			bool bKeepCalculating = false;
+
+			UWorld* World = GCurrentLevelEditingViewportClient ? GCurrentLevelEditingViewportClient->GetWorld() : nullptr;
+			if (World && CurrentFramesInfo.CurrentFrameTimes.NumFrames > 0)
+			{
+				bKeepCalculating = CurrentFramesInfo.KeepCalculating();
+
+				const TArray<int32>& IndicesToCalcluate = CurrentFramesInfo.IndicesToCalculate();
+				if (IndicesToCalcluate.Num() > 0)
 				{
-					if (!AllTrails.Contains(TrailControlTransform.ElementGuid))
-					{
-						continue;
-					}
+					UE::AIE::FEvalHelpers::CalculateWorldTransforms(World, SequencerPtr.Get(), CurrentFramesInfo.CurrentFrameTimes,
+						IndicesToCalcluate, EvaluatingActors, EvaluatingControlRigs, DependentTracks);
+				}
 
-					FTrail* Trail = AllTrails[TrailControlTransform.ElementGuid].Get();
-					if (bUseEditedTimes == false && Trail->GetTrajectoryTransforms())
-					{
-						Trail->GetTrajectoryTransforms()->SetTransforms(TrailControlTransform.Transforms, ControlRigParentWorldTransforms);
-					}
-					Trail->UpdateKeysInRange(GetViewRange());
+			}
+			
+			const TRange<FFrameNumber> Range = TickEvalRange;
+			const TArray<int32>& IndicesToCalcluate = CurrentFramesInfo.IndicesToCalculate();
+			for (const FGuid& CurGuid : EvaluatingTrails)
+			{
+				if (const TUniquePtr<FTrail>* Trail = GetAllTrails().Find(CurGuid))
+				{
+					(*Trail)->UpdateFinished(Range, IndicesToCalcluate);
 				}
 			}
-		}
-		if (Player)
-		{
-			FFrameTime StartTime = Sequencer->GetGlobalTime().Time;
-			FMovieSceneContext Context = FMovieSceneContext(FMovieSceneEvaluationRange(StartTime, TickResolution), Player->GetPlaybackStatus()).SetHasJumped(true);
-			Player->GetEvaluationTemplate().EvaluateSynchronousBlocking(Context);
-			ControlRig->Evaluate_AnyThread();
+
+			if (bKeepCalculating == false)
+			{
+				EvaluatingActors.Reset();
+				EvaluatingControlRigs.Reset();
+				EvaluatingTrails.Reset();
+			}
 		}
 	}
 }
 
-void FSequencerTrailHierarchy::UpdateControlRig(const FTrailEvaluateTimes& EvaluateTimes,UControlRig* ControlRig, TMap<FName, FGuid > &CompMapPair)
+void FSequencerTrailHierarchy::EvaluateActor(const FGuid& InGuid, FActorForWorldTransforms& Actor, TSharedPtr<FArrayOfTransforms>& WorldTransforms, TSharedPtr<FArrayOfTransforms>& ParentWorldTransforms)
 {
-	if (TSharedPtr<IControlRigObjectBinding> ObjectBinding = ControlRig->GetObjectBinding())
+	if (EvaluatingTrails.Contains(InGuid) == false)
 	{
-		TSharedPtr<ISequencer> Sequencer = WeakSequencer.Pin();
-		check(Sequencer);
-		FFrameRate TickResolution = Sequencer->GetFocusedTickResolution();
-		TArray<FFrameNumber> Frames;
-		bool bUseEditedTimes = false;
-
-		const FFrameTime LastFrame = EvaluateTimes.EvalTimes[EvaluateTimes.EvalTimes.Num() - 1] * TickResolution;
-		for (TPair<FName, FGuid >& Pair : CompMapPair)
-		{
-			if (AllTrails[Pair.Value]->GetEditedTimes(this,LastFrame.RoundToFrame(), Frames))
-			{
-				bUseEditedTimes = true;
-				break;
-			}
-		}
-		int32 Index = 0;
-		if (bUseEditedTimes == false)
+		if (const TUniquePtr<FTrail>* Trail = GetAllTrails().Find(InGuid))
 		{
-			//Todo mz make time in frames
-			Frames.SetNum(EvaluateTimes.EvalTimes.Num());
-			for (const double Time : EvaluateTimes.EvalTimes)
-			{
-				const FFrameTime TickTime = Time * TickResolution;
-				Frames[Index++] = TickTime.RoundToFrame();
-			}
+			(*Trail)->HasStartedEvaluating();
 		}
-		UpdateControlRig(Frames, ControlRig, CompMapPair, bUseEditedTimes);
 	}
-}
-void FSequencerTrailHierarchy::Update()
-{
-	const FDateTime UpdateStartTime = FDateTime::Now();
 
-	UpdateViewAndEvalRange();
-	FTrailHierarchy::Update();  //calculates EvalTimesArr
-	FTrailEvaluateTimes EvalTimes = FTrailEvaluateTimes(EvalTimesArr, SecondsPerSegment);
+	//something changed so we reset and recalculate our cache
+	CurrentFramesInfo.Reset();
+
+	FActorAndWorldTransforms ActorAndWorldTransforms(WorldTransforms, ParentWorldTransforms);
+	ActorAndWorldTransforms.Actor = Actor;
+	ActorAndWorldTransforms.SetNumOfTransforms(CurrentFramesInfo.CurrentFrameTimes.NumFrames);
+	EvaluatingActors.Add(ActorAndWorldTransforms);
+	EvaluatingTrails.Add(InGuid);
 
-	for (TPair<UControlRig*, TMap<FName, FGuid>>& CompMapPair : ControlsTracked)
+}
+
+void FSequencerTrailHierarchy::EvaluateControlRig(const FGuid& InGuid, UControlRig* ControlRig, const FName& ControlName, TSharedPtr<FArrayOfTransforms>& WorldTransforms)
+{
+	if (EvaluatingTrails.Contains(InGuid) == false)
 	{
-		bool bNeedToUpdateControlRig = false;
-		FTrail* ForceTrail = nullptr;
-		for (TPair<FName, FGuid>& Pair : CompMapPair.Value)
+		if (const TUniquePtr<FTrail>* Trail = GetAllTrails().Find(InGuid))
 		{
-			if (AllTrails[Pair.Value]->GetCacheState() == ETrailCacheState::Stale)
-			{
-				bNeedToUpdateControlRig = true;
-				ForceTrail = AllTrails[Pair.Value].Get();
-				break;
-			}
+			(*Trail)->HasStartedEvaluating();
 		}
-		if (bNeedToUpdateControlRig == true)
-		{
-			FMovieSceneControlRigTransformTrail* ControlRigTrail = static_cast<FMovieSceneControlRigTransformTrail*>(ForceTrail);
-				
+	}
+	//something changed so we reset
+	CurrentFramesInfo.Reset();
 
-			//if the mouse is not down we update everything like normal making sure we don't use keys to draw the trajectory
-			//if mouse is down we make sure to just use keys to draw the trajectory and we figure out what times we just need to update
-			//which will be selected keys and the current time.
-			if (FSlateApplication::Get().GetPressedMouseButtons().Contains(EKeys::LeftMouseButton) == false)
-			{
-				if (ControlRigTrail)
-				{
-					ControlRigTrail->SetUseKeysForTrajectory(false);
-				}
-				UpdateControlRig(EvalTimes, CompMapPair.Key, CompMapPair.Value);
-			}
-			else if(ForceTrail)
-			{
-				if (ControlRigTrail)
-				{
-					ControlRigTrail->SetUseKeysForTrajectory(true);
-				}
-				//just do selected key times and playback time
-				TArray<FFrameNumber> Frames = ForceTrail->GetSelectedKeyTimes();
-				TSharedPtr<ISequencer> SequencerPtr = WeakSequencer.Pin();
-				FFrameNumber CurrentFrame = SequencerPtr->GetLocalTime().Time.GetFrame();
-				if (Frames.Contains(CurrentFrame) == false)
-				{
-					Frames.Add(CurrentFrame);
-				}
-				UpdateControlRig(Frames, CompMapPair.Key, CompMapPair.Value, true);
-				ForceTrail->ForceEvaluateNextTick();
-			}
+	WorldTransforms->SetNum(CurrentFramesInfo.CurrentFrameTimes.NumFrames);
+	if (FControlRigAndWorldTransforms* Existing = EvaluatingControlRigs.Find(ControlRig))
+	{
+		Existing->SetNumOfTransforms(CurrentFramesInfo.CurrentFrameTimes.NumFrames);
+		TSharedPtr<FArrayOfTransforms>& ExistingArray = Existing->ControlAndWorldTransforms.FindOrAdd(ControlName);
+		ExistingArray = WorldTransforms;
+	}
+	else
+	{
+		if (FControlMapAndTransforms* Map = ControlsTracked.Find(ControlRig))
+		{
+			FControlRigAndWorldTransforms ControlRigAndWorldTransforms;
+			ControlRigAndWorldTransforms.ControlRig = ControlRig;
+			ControlRigAndWorldTransforms.ParentTransforms = Map->ArrayOfTransforms;
+			ControlRigAndWorldTransforms.ParentTransforms->SetNum(CurrentFramesInfo.CurrentFrameTimes.NumFrames);
+			ControlRigAndWorldTransforms.ControlAndWorldTransforms.Add(ControlName, WorldTransforms);
+			EvaluatingControlRigs.Add(ControlRig, ControlRigAndWorldTransforms);
 		}
 	}
-
-	const FTimespan UpdateTimespan = FDateTime::Now() - UpdateStartTime;
-	TimingStats.Add("FSequencerTrailHierarchy::Update", UpdateTimespan);
+	EvaluatingTrails.Add(InGuid);
 }
 
-void FSequencerTrailHierarchy::OnBoneVisibilityChanged(class USkeleton* Skeleton, const FName& BoneName, const bool bIsVisible)
+void FSequencerTrailHierarchy::Update()
 {
-	TSharedPtr<ISequencer> Sequencer = WeakSequencer.Pin();
-	check(Sequencer);
-
-	TArray<FGuid> SelectedSequencerGuids;
-	Sequencer->GetSelectedObjects(SelectedSequencerGuids);
+	const FDateTime UpdateStartTime = FDateTime::Now();
 
-	// TODO: potentially expensive
-	for (const FGuid& SelectedGuid : SelectedSequencerGuids)
-	{
-		for (TWeakObjectPtr<> BoundObject : Sequencer->FindObjectsInCurrentSequence(SelectedGuid))
-		{
-			USkeletalMeshComponent* BoundComponent = Cast<USkeletalMeshComponent>(BoundObject.Get());
-			if (AActor* BoundActor = Cast<AActor>(BoundObject.Get()))
-			{
-				BoundComponent = BoundActor->FindComponentByClass<USkeletalMeshComponent>();
-			}
+	//set up the range
+	UpdateViewAndEvalRange();
+	//remove dead trails and let us know if they need to get updated
+	FTrailHierarchy::Update(); 
 
-			if (!BoundComponent || !BoundComponent->GetSkeletalMeshAsset() || !BoundComponent->GetSkeletalMeshAsset()->GetSkeleton() || !(BoundComponent->GetSkeletalMeshAsset()->GetSkeleton() == Skeleton) || !BonesTracked.Contains(BoundComponent))
-			{
-				continue;
-			}
+	//update new ones
+	EvaluateSequencerAndSetTransforms();
 
-			const FGuid BoneTrailGuid = BonesTracked[BoundComponent][BoneName];
-			const int32 BoneIndex = Skeleton->GetReferenceSkeleton().FindBoneIndex(BoneName);
-			
-			if (bIsVisible)
-			{
-				VisibilityManager.VisibilityMask.Remove(BoneTrailGuid);
-				VisibilityManager.Selected.Add(BoneTrailGuid);
-			}
-			else
-			{
-				VisibilityManager.VisibilityMask.Add(BoneTrailGuid);
-				VisibilityManager.Selected.Remove(BoneTrailGuid);
-			}
-		}
-	}
+	const FTimespan UpdateTimespan = FDateTime::Now() - UpdateStartTime;
+	TimingStats.Add("FSequencerTrailHierarchy::Update", UpdateTimespan);
 }
 
 void FSequencerTrailHierarchy::OnBindingVisibilityStateChanged(UObject* BoundObject, const EBindingVisibilityState VisibilityState)
@@ -444,7 +514,7 @@ void FSequencerTrailHierarchy::OnBindingVisibilityStateChanged(UObject* BoundObj
 		{
 			VisibilityManager.AlwaysVisible.Remove(Guid);
 		}
-	};
+		};
 
 	if (ObjectsTracked.Contains(BoundObject))
 	{
@@ -457,16 +527,20 @@ void FSequencerTrailHierarchy::OnBindingVisibilityStateChanged(UObject* BoundObj
 		return;
 	}
 
-	for (const TPair<FName, FGuid>& Pair : BonesTracked.FindRef(SkelMeshComp))
-	{
-		UpdateTrailVisibilityState(Pair.Value);
-	}
-	/* to do mz
-	for (const TPair<FName, FGuid>& Pair : ControlsTracked.FindRef(SkelMeshComp))
+	for (TPair<UControlRig*, FControlMapAndTransforms>& CompMapPair : ControlsTracked)
 	{
-		UpdateTrailVisibilityState(ControlsTracked[SkelMeshComp][Pair.Key]);
+		if (UControlRig* ControlRig = CompMapPair.Key)
+		{
+			if (ControlRig->GetObjectBinding() && ControlRig->GetObjectBinding()->GetBoundObject() == FControlRigObjectBinding::GetBindableObject(BoundObject))
+			{
+				for (TPair<FName, FGuid> NameGuid : CompMapPair.Value.NameToTrail)
+				{
+					UpdateTrailVisibilityState(NameGuid.Value);
+				}
+				break;
+			}
+		}
 	}
-	*/
 }
 
 void FSequencerTrailHierarchy::UpdateSequencerBindings(const TArray<FGuid>& SequencerBindings, TFunctionRef<void(UObject*, FTrail*, FGuid)> OnUpdated)
@@ -476,9 +550,14 @@ void FSequencerTrailHierarchy::UpdateSequencerBindings(const TArray<FGuid>& Sequ
 	TSharedPtr<ISequencer> Sequencer = WeakSequencer.Pin();
 	check(Sequencer);
 
+	TArray<AActor*> SelectedActors;
+	GEditor->GetSelectedActors()->GetSelectedObjects<AActor>(SelectedActors);
+	TArray<USceneComponent*> SelectedSceneComponents;
+	GEditor->GetSelectedActors()->GetSelectedObjects<USceneComponent>(SelectedSceneComponents);
 	for (FGuid BindingGuid : SequencerBindings)
 	{
-		if (UMovieScene3DTransformTrack* TransformTrack = Sequencer->GetFocusedMovieSceneSequence()->GetMovieScene()->FindTrack<UMovieScene3DTransformTrack>(BindingGuid))
+		bool bAddedControlRig = false;
+		if (UMovieSceneControlRigParameterTrack* CRParameterTrack = Sequencer->GetFocusedMovieSceneSequence()->GetMovieScene()->FindTrack<UMovieSceneControlRigParameterTrack>(BindingGuid))
 		{
 			for (TWeakObjectPtr<> BoundObject : Sequencer->FindBoundObjects(BindingGuid, Sequencer->GetFocusedTemplateID()))
 			{
@@ -486,35 +565,44 @@ void FSequencerTrailHierarchy::UpdateSequencerBindings(const TArray<FGuid>& Sequ
 				{
 					continue;
 				}
-				//if using old trails don't add new ones.
-				if (CVarUseOldSequencerMotionTrails->GetBool() == true)
-				{
-					continue;
-				}
 
-				USceneComponent* BoundComponent = Cast<USceneComponent>(BoundObject.Get());
+				USkeletalMeshComponent* BoundComponent = Cast<USkeletalMeshComponent>(BoundObject.Get());
 				if (AActor* BoundActor = Cast<AActor>(BoundObject.Get()))
 				{
-					BoundComponent = BoundActor->GetRootComponent();
+					BoundComponent = BoundActor->FindComponentByClass<USkeletalMeshComponent>();
 				}
 
-				if (!ObjectsTracked.Contains(BoundComponent))
+				if (!BoundComponent || !BoundComponent->GetSkeletalMeshAsset() || !BoundComponent->GetSkeletalMeshAsset()->GetSkeleton())
 				{
-					AddComponentToHierarchy(BoundComponent, TransformTrack);
-				}
+					continue;
+				};
 
-				if (!ObjectsTracked.Contains(BoundComponent))
+				//if we have a selected control rig don't add the transform track also, makes pinning messy
+				if (!ControlRigDelegateHandles.Contains(CRParameterTrack))
 				{
-					continue;
+					RegisterControlRigDelegates(BoundComponent, CRParameterTrack);
+					UControlRig* ControlRig = CRParameterTrack->GetControlRig();
+					if (ControlRig)
+					{
+						TArray<FName> Selected = ControlRig->CurrentControlSelection();
+						for (const FName& ControlName : Selected)
+						{
+							AddControlRigTrail(BoundComponent, ControlRig, CRParameterTrack, ControlName);
+							bAddedControlRig = true; 
+						}
+					}
 				}
-
-				if (AllTrails.Contains(ObjectsTracked[BoundComponent]) && AllTrails[ObjectsTracked[BoundComponent]].IsValid())
+				if (bAddedControlRig)
 				{
-					OnUpdated(BoundComponent, AllTrails[ObjectsTracked[BoundComponent]].Get(), ObjectsTracked[BoundComponent]);
+					continue;
 				}
 			}
-		} // if TransformTrack
-		if (UMovieSceneSkeletalAnimationTrack* AnimTrack = Sequencer->GetFocusedMovieSceneSequence()->GetMovieScene()->FindTrack<UMovieSceneSkeletalAnimationTrack>(BindingGuid))
+		} // if ControlRigParameterTrack
+		if (bAddedControlRig)
+		{
+			continue;
+		}	
+		if (UMovieScene3DTransformTrack* TransformTrack = Sequencer->GetFocusedMovieSceneSequence()->GetMovieScene()->FindTrack<UMovieScene3DTransformTrack>(BindingGuid))
 		{
 			for (TWeakObjectPtr<> BoundObject : Sequencer->FindBoundObjects(BindingGuid, Sequencer->GetFocusedTemplateID()))
 			{
@@ -522,55 +610,46 @@ void FSequencerTrailHierarchy::UpdateSequencerBindings(const TArray<FGuid>& Sequ
 				{
 					continue;
 				}
+			
+				/*
+				//if using old trails don't add new ones.
 				if (CVarUseOldSequencerMotionTrails->GetBool() == true)
 				{
 					continue;
 				}
+				*/
 
-				USkeletalMeshComponent* BoundComponent = Cast<USkeletalMeshComponent>(BoundObject.Get());
+				USceneComponent* BoundComponent = Cast<USceneComponent>(BoundObject.Get());
 				if (AActor* BoundActor = Cast<AActor>(BoundObject.Get()))
 				{
-					BoundComponent = BoundActor->FindComponentByClass<USkeletalMeshComponent>();
-				}
-
-				if (!BoundComponent || !BoundComponent->GetSkeletalMeshAsset() || !BoundComponent->GetSkeletalMeshAsset()->GetSkeleton())
-				{
-					continue;
-				}
-
-				if (!ObjectsTracked.Contains(BoundComponent))
-				{
-					if (UMovieScene3DTransformTrack* TransformTrack = Sequencer->GetFocusedMovieSceneSequence()->GetMovieScene()->FindTrack<UMovieScene3DTransformTrack>(BindingGuid))
+					if (SelectedActors.Contains(BoundActor) == false)
 					{
-						AddComponentToHierarchy(BoundComponent, TransformTrack);
+						continue;
 					}
+					BoundComponent = BoundActor->GetRootComponent();
 				}
-
-				if (!ObjectsTracked.Contains(BoundComponent))
+				else if (SelectedSceneComponents.Contains(BoundComponent) == false)
 				{
 					continue;
 				}
 
-				if (!BonesTracked.Contains(BoundComponent))
+				if (!ObjectsTracked.Contains(BoundComponent))
 				{
-					//AddSkeletonToHierarchy(BoundComponent); //need to revisit skeletals
+					AddComponentToHierarchy(BindingGuid, BoundComponent, TransformTrack);
 				}
 
-				if (!BonesTracked.Contains(BoundComponent))
+				if (!ObjectsTracked.Contains(BoundComponent))
 				{
 					continue;
 				}
 
-				for (const TPair<FName, FGuid>& BoneNameGuidPair : BonesTracked[BoundComponent])
+				if (AllTrails.Contains(ObjectsTracked[BoundComponent]) && AllTrails[ObjectsTracked[BoundComponent]].IsValid())
 				{
-					if (AllTrails.Contains(BoneNameGuidPair.Value) && AllTrails[BoneNameGuidPair.Value].IsValid())
-					{
-						OnUpdated(BoundComponent, AllTrails[BoneNameGuidPair.Value].Get(), BoneNameGuidPair.Value);
-					}
+					OnUpdated(BoundComponent, AllTrails[ObjectsTracked[BoundComponent]].Get(), ObjectsTracked[BoundComponent]);
 				}
 			}
-		}
-		if (UMovieSceneControlRigParameterTrack* CRParameterTrack = Sequencer->GetFocusedMovieSceneSequence()->GetMovieScene()->FindTrack<UMovieSceneControlRigParameterTrack>(BindingGuid))
+		} // if TransformTrack
+		if (UMovieSceneSkeletalAnimationTrack* AnimTrack = Sequencer->GetFocusedMovieSceneSequence()->GetMovieScene()->FindTrack<UMovieSceneSkeletalAnimationTrack>(BindingGuid))
 		{
 			for (TWeakObjectPtr<> BoundObject : Sequencer->FindBoundObjects(BindingGuid, Sequencer->GetFocusedTemplateID()))
 			{
@@ -578,40 +657,63 @@ void FSequencerTrailHierarchy::UpdateSequencerBindings(const TArray<FGuid>& Sequ
 				{
 					continue;
 				}
+				/*
+				if (CVarUseOldSequencerMotionTrails->GetBool() == true)
+				{
+					continue;
+				}
+				*/
 
 				USkeletalMeshComponent* BoundComponent = Cast<USkeletalMeshComponent>(BoundObject.Get());
 				if (AActor* BoundActor = Cast<AActor>(BoundObject.Get()))
 				{
+					if (SelectedActors.Contains(BoundActor) == false)
+					{
+						continue;
+					}
 					BoundComponent = BoundActor->FindComponentByClass<USkeletalMeshComponent>();
 				}
+				else if (SelectedSceneComponents.Contains(BoundComponent) == false)
+				{
+					continue;
+				}
 
 				if (!BoundComponent || !BoundComponent->GetSkeletalMeshAsset() || !BoundComponent->GetSkeletalMeshAsset()->GetSkeleton())
 				{
 					continue;
-				};
+				}
 
-				if (!ControlRigDelegateHandles.Contains(CRParameterTrack))
+				
+
+				if (!ObjectsTracked.Contains(BoundComponent))
 				{
-					RegisterControlRigDelegates(BoundComponent, CRParameterTrack);
-					UControlRig* ControlRig = CRParameterTrack->GetControlRig();
-					if (ControlRig)
+					if (UMovieScene3DTransformTrack* TransformTrack = Sequencer->GetFocusedMovieSceneSequence()->GetMovieScene()->FindTrack<UMovieScene3DTransformTrack>(BindingGuid))
 					{
-						TArray<FName> Selected = ControlRig->CurrentControlSelection();
-						for (const FName& ControlName : Selected)
-						{
-							AddControlRigTrail(BoundComponent,ControlRig, CRParameterTrack, ControlName);
-						}
+						AddComponentToHierarchy(BindingGuid,BoundComponent, TransformTrack);
 					}
 				}
 			}
-		} // if ControlRigParameterTrack
-
+		}
 	}
 	const FTimespan Timespan = FDateTime::Now() - StartTime;
 	TimingStats.Add("FSequencerTrailHierarchy::UpdateSequencerBindings", Timespan);
 }
 
 
+bool FSequencerTrailHierarchy::CheckForChanges()
+{
+	bool bHasChange = false;
+	if (TSharedPtr<ISequencer> Sequencer = WeakSequencer.Pin())
+	{
+		if (Sequencer->GetFocusedMovieSceneSequence() && Sequencer->GetFocusedMovieSceneSequence()->GetMovieScene())
+		{
+			bHasChange = LastValidMovieSceneGuid != Sequencer->GetFocusedMovieSceneSequence()->GetMovieScene()->GetSignature();
+			LastValidMovieSceneGuid = Sequencer->GetFocusedMovieSceneSequence()->GetMovieScene()->GetSignature();
+		}
+	}
+	return bHasChange;
+}
+
 void FSequencerTrailHierarchy::UpdateViewAndEvalRange()
 {
 	TSharedPtr<ISequencer> Sequencer = WeakSequencer.Pin();
@@ -620,8 +722,7 @@ void FSequencerTrailHierarchy::UpdateViewAndEvalRange()
 	FFrameRate TickResolution = Sequencer->GetFocusedTickResolution();
 	FFrameRate DisplayRate = Sequencer->GetFocusedDisplayRate();
 
-	TRange<FFrameNumber> TickViewRange;
-	TRange<FFrameNumber> EvalViewRange  = Sequencer->GetFocusedMovieSceneSequence()->GetMovieScene()->GetPlaybackRange();
+	TickEvalRange  = Sequencer->GetFocusedMovieSceneSequence()->GetMovieScene()->GetPlaybackRange();
 	if (!UMotionTrailToolOptions::GetTrailOptions()->bShowFullTrail)
 	{
 		FFrameTime SequenceTime = Sequencer->GetLocalTime().Time;
@@ -633,29 +734,16 @@ void FSequencerTrailHierarchy::UpdateViewAndEvalRange()
 	{
 		TickViewRange = Sequencer->GetFocusedMovieSceneSequence()->GetMovieScene()->GetPlaybackRange();
 	}
-
-	const double StartViewSeconds = TickResolution.AsSeconds(FFrameTime(TickViewRange.GetLowerBoundValue()));
-	const double EndViewSeconds = TickResolution.AsSeconds(FFrameTime(TickViewRange.GetUpperBoundValue()));
-	ViewRange = TRange<double>(StartViewSeconds,EndViewSeconds);
-
-	const double StartEvalSeconds = TickResolution.AsSeconds(FFrameTime(EvalViewRange.GetLowerBoundValue()));
-	const double EndEvalSeconds = TickResolution.AsSeconds(FFrameTime(EvalViewRange.GetUpperBoundValue()));
-	EvalRange = TRange<double>(StartEvalSeconds, EndEvalSeconds);
-
-	// snap view range to ticks per segment
-	//const double TicksBetween = FMath::Fmod(StartSeconds, GetSecondsPerSegment());
-	//ViewRange = TRange<double>(StartSeconds - TicksBetween, EndSeconds - TicksBetween);
-
 }
 
-void FSequencerTrailHierarchy::AddComponentToHierarchy(USceneComponent* CompToAdd, UMovieScene3DTransformTrack* TransformTrack)
+FGuid FSequencerTrailHierarchy::AddComponentToHierarchy(const FGuid& InBindingGuid, USceneComponent* CompToAdd, UMovieScene3DTransformTrack* TransformTrack)
 {
 	TSharedPtr<ISequencer> Sequencer = WeakSequencer.Pin();
 	check(Sequencer);
 	
 	const FGuid CurTrailGuid = ObjectsTracked.FindOrAdd(CompToAdd, FGuid::NewGuid());
 
-	TUniquePtr<FTrail> CurTrail = MakeUnique<FMovieSceneComponentTransformTrail>(CompToAdd,UMotionTrailToolOptions::GetTrailOptions()->TrailColor, false, TransformTrack, Sequencer);
+	TUniquePtr<FTrail> CurTrail = MakeUnique<FMovieSceneComponentTransformTrail>(InBindingGuid, CompToAdd, false, TransformTrack, Sequencer);
 	if (AllTrails.Contains(ObjectsTracked[CompToAdd])) 
 	{
 		AllTrails.Remove(ObjectsTracked[CompToAdd]);
@@ -663,58 +751,51 @@ void FSequencerTrailHierarchy::AddComponentToHierarchy(USceneComponent* CompToAd
 	CurTrail->ForceEvaluateNextTick();
 
 	AddTrail(ObjectsTracked[CompToAdd], MoveTemp(CurTrail));
+	return CurTrailGuid;
 }
 
-void FSequencerTrailHierarchy::AddSkeletonToHierarchy(class USkeletalMeshComponent* CompToAdd)
+FGuid FSequencerTrailHierarchy::AddControlRigTrail(USkeletalMeshComponent* Component, UControlRig* ControlRig, UMovieSceneControlRigParameterTrack* CRParameterTrack,const FName& ControlName)
 {
-	const FDateTime StartTime = FDateTime::Now();
-
-	TSharedPtr<ISequencer> Sequencer = WeakSequencer.Pin();
-	check(Sequencer);
-
-	TSharedPtr<FAnimTrajectoryCache> AnimTrajectoryCache = MakeShared<FAnimTrajectoryCache>(CompToAdd, Sequencer);
-	TMap<FName, FGuid>& BoneMap = BonesTracked.Add(CompToAdd, TMap<FName, FGuid>());
-	
-	USkeleton* MySkeleton = CompToAdd->GetSkeletalMeshAsset()->GetSkeleton();
-	const int32 NumBones = MySkeleton->GetReferenceSkeleton().GetNum();
-	for (int32 BoneIdx = 0; BoneIdx < NumBones; BoneIdx++)
+	FGuid NewGuid;
+	if (FRigControlElement* ControlElement = ControlRig->FindControl(ControlName))
 	{
-		int32 ParentBoneIndex = MySkeleton->GetReferenceSkeleton().GetParentIndex(BoneIdx);
-		const FName BoneName = MySkeleton->GetReferenceSkeleton().GetBoneName(BoneIdx);
-
-		const FGuid BoneGuid = FGuid::NewGuid();
-		BoneMap.Add(BoneName, BoneGuid);
-		VisibilityManager.VisibilityMask.Add(BoneGuid);
-
-		if (ParentBoneIndex != INDEX_NONE)
+		if (ControlElement->IsAnimationChannel())
 		{
-			AllTrails.Add(BoneGuid, MakeUnique<FAnimationBoneTrail>(CompToAdd,UMotionTrailToolOptions::GetTrailOptions()->TrailColor, false, AnimTrajectoryCache, BoneName, false));
-
-		}
-		else // root bone
-		{
-			AllTrails.Add(BoneGuid, MakeUnique<FAnimationBoneTrail>(CompToAdd,UMotionTrailToolOptions::GetTrailOptions()->TrailColor, false, AnimTrajectoryCache, BoneName, true));
+			return NewGuid; // no shape
 		}
 	}
-
-	const FTimespan Timespan = FDateTime::Now() - StartTime;
-	TimingStats.Add("FSequencerTrailHierarchy::AddSkeletonToHierarchy", Timespan);
-}
-
-void FSequencerTrailHierarchy::AddControlRigTrail(USkeletalMeshComponent* Component, UControlRig* ControlRig, UMovieSceneControlRigParameterTrack* CRParameterTrack,const FName& ControlName)
-{
 	if (ControlsTracked.Find(ControlRig) == nullptr)
 	{
-		ControlsTracked.Add(ControlRig, TMap<FName, FGuid>());
+		FControlMapAndTransforms MapAndTransforms;
+		MapAndTransforms.ArrayOfTransforms = MakeShared<FArrayOfTransforms>();
+		
+		ControlsTracked.Add(ControlRig, MapAndTransforms);
 	}
-	if (!ControlsTracked[ControlRig].Contains(ControlName))
+	if (!ControlsTracked[ControlRig].NameToTrail.Contains(ControlName))
 	{
 		TSharedPtr<ISequencer> Sequencer = WeakSequencer.Pin();
 		if (Sequencer)
 		{
-			const FGuid ControlGuid = ControlsTracked.Find(ControlRig)->FindOrAdd(ControlName, FGuid::NewGuid());
-
-			TUniquePtr<FTrail> CurTrail = MakeUnique<FMovieSceneControlRigTransformTrail>(Component, UMotionTrailToolOptions::GetTrailOptions()->TrailColor, false, CRParameterTrack, Sequencer, ControlName);
+			//find binding id
+			TArray<UMovieSceneControlRigParameterTrack*> Tracks;
+			UMovieScene* MovieScene = Sequencer->GetFocusedMovieSceneSequence()->GetMovieScene();
+			const TArray<FMovieSceneBinding>& Bindings = MovieScene->GetBindings();
+			FGuid BindingID;
+			for (const FMovieSceneBinding& Binding : Bindings)
+			{
+				TArray<UMovieSceneTrack*> FoundTracks = MovieScene->FindTracks(UMovieSceneControlRigParameterTrack::StaticClass(), Binding.GetObjectGuid(), NAME_None);
+				for (UMovieSceneTrack* Track : FoundTracks)
+				{
+					if (UMovieSceneControlRigParameterTrack* CRTrack = Cast<UMovieSceneControlRigParameterTrack>(Track))
+					{
+						BindingID = Binding.GetObjectGuid();
+						break;
+					}
+				}
+			}
+			const FGuid ControlGuid = ControlsTracked.Find(ControlRig)->NameToTrail.FindOrAdd(ControlName, FGuid::NewGuid());
+			TUniquePtr<FTrail> CurTrail = MakeUnique<FMovieSceneControlRigTransformTrail>(BindingID,Component, false, CRParameterTrack, Sequencer,
+				ControlName, ControlsTracked[ControlRig].ArrayOfTransforms);
 			if (AllTrails.Contains(ControlGuid))
 			{
 				AllTrails.Remove(ControlGuid);
@@ -722,16 +803,19 @@ void FSequencerTrailHierarchy::AddControlRigTrail(USkeletalMeshComponent* Compon
 			}
 			AddTrail(ControlGuid, MoveTemp(CurTrail));
 			VisibilityManager.ControlSelected.Add(ControlGuid);
+			NewGuid = ControlGuid;
 		}
 	}
 	else
 	{
-		const FGuid* ControlGuid = ControlsTracked.Find(ControlRig)->Find(ControlName);
+		const FGuid* ControlGuid = ControlsTracked.Find(ControlRig)->NameToTrail.Find(ControlName);
 		if (ControlGuid != nullptr)
 		{
 			VisibilityManager.ControlSelected.Add(*ControlGuid);
+			NewGuid = *ControlGuid;
 		}
 	}
+	return NewGuid;
 }
 
 void FSequencerTrailHierarchy::RegisterControlRigDelegates(USkeletalMeshComponent* Component, UMovieSceneControlRigParameterTrack* CRParameterTrack)
@@ -755,11 +839,11 @@ void FSequencerTrailHierarchy::RegisterControlRigDelegates(USkeletalMeshComponen
 
 			if (ControlsTracked.Find(ControlRig) != nullptr)
 			{
-				if (ControlsTracked[ControlRig].Contains(ControlElement->GetFName()))
+				if (ControlsTracked[ControlRig].NameToTrail.Contains(ControlElement->GetFName()))
 				{
 					if (bSelected == false)
 					{
-						const FGuid TrailGuid = ControlsTracked[ControlRig][ControlElement->GetFName()];
+						const FGuid TrailGuid = ControlsTracked[ControlRig].NameToTrail[ControlElement->GetFName()];
 						VisibilityManager.ControlSelected.Remove(TrailGuid);
 						RemoveTrailIfNotAlwaysVisible(TrailGuid);
 					}
@@ -768,10 +852,10 @@ void FSequencerTrailHierarchy::RegisterControlRigDelegates(USkeletalMeshComponen
 
 			//check to see if the seleced control rig is sill selected
 			TArray<FGuid> TrailsToRemove;
-			for (TPair<UControlRig*, TMap<FName, FGuid>>& CompMapPair : ControlsTracked)
-			{
+			for (TPair<UControlRig*, FControlMapAndTransforms>& CompMapPair : ControlsTracked)
+			{			
 				UControlRig* TrackedControlRig = CompMapPair.Key;
-				for (TPair<FName, FGuid> NameGuid : CompMapPair.Value)
+				for (TPair<FName, FGuid> NameGuid : CompMapPair.Value.NameToTrail)
 				{
 					if (TrackedControlRig->IsControlSelected(NameGuid.Key) == false)
 					{
@@ -803,30 +887,221 @@ void FSequencerTrailHierarchy::RegisterControlRigDelegates(USkeletalMeshComponen
 			
 			if(InNotif == ERigHierarchyNotification::ElementRemoved)
 			{
-				if (!ControlsTracked.Contains(ControlRig) || !ControlsTracked[ControlRig].Contains(ControlElement->GetFName())) // We only care about controls
+				if (!ControlsTracked.Contains(ControlRig) || !ControlsTracked[ControlRig].NameToTrail.Contains(ControlElement->GetFName())) // We only care about controls
 				{
 					return;
 				}
 				
-				const FGuid TrailGuid = ControlsTracked[ControlRig][ControlElement->GetFName()];
+				const FGuid TrailGuid = ControlsTracked[ControlRig].NameToTrail[ControlElement->GetFName()];
 				RemoveTrail(TrailGuid);
 			}
             else if(InNotif == ERigHierarchyNotification::ElementRenamed)
             {
             	const FName OldName = InHierarchy->GetPreviousName(ControlElement->GetKey());
 
-            	if (!ControlsTracked.Contains(ControlRig) || !ControlsTracked[ControlRig].Contains(OldName))
+            	if (!ControlsTracked.Contains(ControlRig) || !ControlsTracked[ControlRig].NameToTrail.Contains(OldName))
 				{
 					return;
 				}
 
-				const FGuid TempTrailGuid = ControlsTracked[ControlRig][OldName];
-				ControlsTracked[ControlRig].Remove(OldName);
-				ControlsTracked[ControlRig].Add(ControlElement->GetFName(), TempTrailGuid);
+				const FGuid TempTrailGuid = ControlsTracked[ControlRig].NameToTrail[OldName];
+				ControlsTracked[ControlRig].NameToTrail.Remove(OldName);
+				ControlsTracked[ControlRig].NameToTrail.Add(ControlElement->GetFName(), TempTrailGuid);
             }
         }
     );
 }
 
+void FSequencerTrailHierarchy::RegisterMotionTrailOptionDelegates()
+{
+	if (UMotionTrailToolOptions* TrailOptions = UMotionTrailToolOptions::GetTrailOptions())
+	{
+		TrailOptions->OnPinSelection.AddRaw(this, &FSequencerTrailHierarchy::OnPinSelection);
+		TrailOptions->OnPinComponent.AddRaw(this, &FSequencerTrailHierarchy::OnPinComponent);
+		TrailOptions->OnDeletePinned.AddRaw(this, &FSequencerTrailHierarchy::OnDeletePinned);
+		TrailOptions->OnPutPinnedInSpace.AddRaw(this, &FSequencerTrailHierarchy::OnPutPinnedInSpace);
+		TrailOptions->OnSetLinearColor.AddRaw(this, &FSequencerTrailHierarchy::OnSetLinearColor);
+		TrailOptions->OnSetHasOffset.AddRaw(this, &FSequencerTrailHierarchy::OnSetHasOffset);
+	}
+}
+
+void FSequencerTrailHierarchy::UnRegisterMotionTrailOptionDelegates()
+{
+	if (UMotionTrailToolOptions* TrailOptions = UMotionTrailToolOptions::GetTrailOptions())
+	{
+		TrailOptions->OnPinSelection.RemoveAll(this);
+		TrailOptions->OnPinComponent.RemoveAll(this);
+		TrailOptions->OnDeletePinned.RemoveAll(this);
+		TrailOptions->OnPutPinnedInSpace.RemoveAll(this);
+		TrailOptions->OnSetLinearColor.RemoveAll(this);
+		TrailOptions->OnSetHasOffset.RemoveAll(this);
+	}
+}
+
+void FSequencerTrailHierarchy::PinTrail(FGuid InGuid)
+{
+	if (const TUniquePtr<FTrail>* Trail = GetAllTrails().Find(InGuid))
+	{
+		VisibilityManager.SetTrailAlwaysVisible(InGuid, true);
+		if ((*Trail)->GetDrawInfo())
+		{
+			(*Trail)->GetDrawInfo()->SetColor(UMotionTrailToolOptions::GetTrailOptions()->DefaultColor);
+			(*Trail)->GetDrawInfo()->SetStyle(UMotionTrailToolOptions::GetTrailOptions()->TrailStyle);
+
+			UMotionTrailToolOptions* Settings = UMotionTrailToolOptions::GetTrailOptions();
+			UMotionTrailToolOptions::FPinnedTrail PinnedTrail;
+			PinnedTrail.TrailGuid = InGuid;
+			PinnedTrail.TrailName = (*Trail)->GetName();
+			PinnedTrail.TrailColor = (*Trail)->GetDrawInfo()->GetColor();
+			PinnedTrail.bHasOffset = (*Trail)->HasOffsetTransform();
+			Settings->AddPinned(PinnedTrail);
+		}
+	}
+}
+
+void FSequencerTrailHierarchy::OnPinSelection()
+{
+	for (const TPair<FGuid, TUniquePtr<FTrail>>& GuidTrailPair : GetAllTrails())
+	{
+		if (VisibilityManager.IsTrailAlwaysVisible(GuidTrailPair.Key) == false) //if not always visible then it's selected
+		{
+			PinTrail(GuidTrailPair.Key);
+		}
+	}
+}
+
+void FSequencerTrailHierarchy::OnPinComponent(USceneComponent* InSceneComponent, FName InSocketName)
+{
+	PinComponent(InSceneComponent, InSocketName);
+}
+
+FGuid FSequencerTrailHierarchy::PinComponent(USceneComponent* InSceneComponent, FName InSocketName)
+{
+	TSharedPtr<ISequencer> Sequencer = WeakSequencer.Pin();
+	check(Sequencer);
+	TArray<FMovieSceneObjectBindingID> Bindings;
+	Sequencer->GetEvaluationState()->FilterObjectBindings(InSceneComponent, *Sequencer, &Bindings);
+	if (Bindings.Num() <= 0)
+	{
+		AActor* ParentActor = InSceneComponent->GetOwner();
+		Sequencer->GetEvaluationState()->FilterObjectBindings(ParentActor, *Sequencer, &Bindings);
+
+	}
+	UMovieScene* MovieScene = Sequencer->GetFocusedMovieSceneSequence()->GetMovieScene();
+	TWeakObjectPtr<UMovieSceneTrack> MainTrack = nullptr;
+	if (Bindings.Num() > 0)
+	{
+		const FMovieSceneObjectBindingID Binding = Bindings[0];
+		TArray<UMovieSceneTrack*> FoundCRTracks = MovieScene->FindTracks(UMovieSceneControlRigParameterTrack::StaticClass(), Binding.GetGuid(), NAME_None);
+		TArray<UMovieSceneTrack*> FoundSkelAnimTracks = MovieScene->FindTracks(UMovieSceneSkeletalAnimationTrack::StaticClass(), Binding.GetGuid(), NAME_None);
+		TArray<UMovieSceneTrack*> FoundTransformTracks = MovieScene->FindTracks(UMovieScene3DTransformTrack::StaticClass(), Binding.GetGuid(), NAME_None);
+
+		if (FoundCRTracks.Num() > 0)
+		{
+			MainTrack = FoundCRTracks[0];
+		}
+		else if (FoundSkelAnimTracks.Num() > 0 && FoundSkelAnimTracks[0]->GetAllSections().Num() > 0)
+		{
+			MainTrack = FoundSkelAnimTracks[0];
+		}
+		else if (FoundTransformTracks.Num() > 0)
+		{
+			MainTrack = FoundTransformTracks[0];
+		}
+	}
+	if (MainTrack.IsValid() == false) //no binding or no tracks so bail
+	{
+		UE_LOG(LogTemp, Warning, TEXT("MotionTrails: No binding or tracks for pinned component. Please make sure it is added to Sequencer."));
+		return FGuid();
+	}
+
+	TMap<FName, FGuid>& BoneNameGuidPair = SocketsTracked.FindOrAdd(InSceneComponent);
+	const FGuid* Guid = BoneNameGuidPair.Find(InSocketName);
+	if (Guid)
+	{
+		if (AllTrails.Contains(*Guid))
+		{
+			AllTrails.Remove(*Guid);
+		}
+		VisibilityManager.AlwaysVisible.Remove(*Guid);
+	}
+	
+	const FMovieSceneObjectBindingID Binding = Bindings[0];
+	TUniquePtr<FTrail> CurTrail = MakeUnique<FMovieSceneSocketTransformTrail>(Binding.GetGuid(), InSceneComponent, InSocketName, false, MainTrack, Sequencer);
+	
+	CurTrail->ForceEvaluateNextTick();
+	const FGuid BoneGuid = FGuid::NewGuid();
+	BoneNameGuidPair.Add(InSocketName, BoneGuid);
+	VisibilityManager.AlwaysVisible.Add(BoneGuid);
+	if (UMotionTrailToolOptions* Settings = UMotionTrailToolOptions::GetTrailOptions())
+	{
+		if (CurTrail->GetDrawInfo())
+		{
+			CurTrail->GetDrawInfo()->SetColor(UMotionTrailToolOptions::GetTrailOptions()->DefaultColor);
+			CurTrail->GetDrawInfo()->SetStyle(UMotionTrailToolOptions::GetTrailOptions()->TrailStyle);
+
+			UMotionTrailToolOptions::FPinnedTrail Trail;
+			Trail.TrailGuid = BoneGuid;
+			Trail.TrailName = CurTrail->GetName();
+			Trail.TrailColor = CurTrail->GetDrawInfo()->GetColor();
+			Trail.bHasOffset = CurTrail->HasOffsetTransform();
+			Settings->AddPinned(Trail);
+		}
+	}
+
+	AddTrail(BoneGuid, MoveTemp(CurTrail));
+	return BoneGuid;
+}
+
+void FSequencerTrailHierarchy::OnDeletePinned(FGuid InGuid)
+{
+	VisibilityManager.SetTrailAlwaysVisible(InGuid, false);
+	RemoveTrail(InGuid);
+}
+
+void FSequencerTrailHierarchy::OnPutPinnedInSpace(FGuid InGuid, AActor* InActor, FName InComponentName)
+{
+	if (const TUniquePtr<FTrail>* Trail = GetAllTrails().Find(InGuid))
+	{
+		if (InActor)
+		{
+			(*Trail)->SetSpace(InActor, InComponentName);
+			(*Trail)->ForceEvaluateNextTick();
+
+		}
+		else
+		{
+			(*Trail)->ClearSpace();
+			(*Trail)->ForceEvaluateNextTick();
+		}
+	}
+}
+
+void FSequencerTrailHierarchy::OnSetLinearColor(FGuid InGuid, FLinearColor Color)
+{
+	if (const TUniquePtr<FTrail>* Trail = GetAllTrails().Find(InGuid))
+	{
+		FTrajectoryDrawInfo* CurDrawInfo = (*Trail)->GetDrawInfo();
+		CurDrawInfo->SetColor(Color);
+	}
+}
+void FSequencerTrailHierarchy::OnSetHasOffset(FGuid InGuid, bool bOffset)
+{
+	if (const TUniquePtr<FTrail>* Trail = GetAllTrails().Find(InGuid))
+	{
+		if (bOffset == false)
+		{
+			(*Trail)->ClearOffsetTransform();
+		}
+		else
+		{
+			(*Trail)->SetOffsetMode();
+		}
+
+	}
+}
+
 } // namespace MovieScene
 } // namespace UE
+
+
diff --git a/Engine/Plugins/Editor/SequencerAnimTools/Source/SequencerAnimTools/Private/Sequencer/SequencerTrailHierarchy.h b/Engine/Plugins/Editor/SequencerAnimTools/Source/SequencerAnimTools/Private/Sequencer/SequencerTrailHierarchy.h
--- a/Engine/Plugins/Editor/SequencerAnimTools/Source/SequencerAnimTools/Private/Sequencer/SequencerTrailHierarchy.h
+++ b/Engine/Plugins/Editor/SequencerAnimTools/Source/SequencerAnimTools/Private/Sequencer/SequencerTrailHierarchy.h
@@ -4,6 +4,7 @@
 
 #include "TrailHierarchy.h"
 #include "ISequencer.h"
+#include "Tools/EvaluateSequencerTools.h"
 
 namespace UE::SequencerAnimTools { class FTrail; }
 namespace UE::SequencerAnimTools { class FTrailEvaluateTimes; }
@@ -13,13 +14,13 @@ class ISequencer;
 class USceneComponent;
 class USkeletalMeshComponent;
 class USkeleton;
+class UMovieScene;
 class UMovieSceneSection;
 class UMovieScene3DTransformTrack;
 class UMovieSceneControlRigParameterTrack;
 class UMovieSceneControlRigParameterSection;
 class UControlRig;
 struct FRigHierarchyContainer;
-
 namespace UE
 {
 namespace SequencerAnimTools
@@ -46,48 +47,57 @@ public:
 		, ControlRigDelegateHandles()
 	{
 	}
+	//New
+	void EvaluateActor(const FGuid& InKey, FActorForWorldTransforms& InActor,TSharedPtr<UE::AIE::FArrayOfTransforms>& WorldTransforms, TSharedPtr<UE::AIE::FArrayOfTransforms>& ParentWorldTransforms);
+	void EvaluateControlRig(const FGuid& InKey, UControlRig* InControlRig, const FName& InControlName, TSharedPtr<UE::AIE::FArrayOfTransforms>& WorldTransforms);
 
 	// FTrailHierarchy interface
 	virtual void Initialize() override;
 	virtual void Destroy() override;
 	virtual ITrailHierarchyRenderer* GetRenderer() const override { return HierarchyRenderer.Get(); }
-	virtual double GetSecondsPerFrame() const override { return 1.0 / WeakSequencer.Pin()->GetFocusedDisplayRate().AsDecimal(); }
-	virtual double GetSecondsPerSegment() const override;
 	virtual FFrameNumber GetFramesPerFrame() const override;
 	virtual FFrameNumber GetFramesPerSegment() const override;
+	virtual const FCurrentFramesInfo* GetCurrentFramesInfo() const override;
+	virtual bool CheckForChanges() override;
 
-	virtual void RemoveTrail(const FGuid& Key) override;
+	virtual void RemoveTrail(const FGuid& InTrailGuid) override;
 	virtual void Update() override;
-	// End FTrailHierarchy interface
+	virtual void CalculateEvalRangeArray() override;
+	virtual bool IsTrailEvaluating(const FGuid& InTrailGuid, bool bIndirectlyOnly) const override;
 
-	const TMap<UObject*, FGuid>& GetObjectsTracked() const { return ObjectsTracked; }
-	const TMap<USkeletalMeshComponent*, TMap<FName, FGuid>>& GetBonesTracked() const { return BonesTracked; }
-	const TMap<UControlRig*, TMap<FName, FGuid>>& GetControlsTracked() const { return ControlsTracked; }
+	// End FTrailHierarchy interface
 
-	void OnBoneVisibilityChanged(USkeleton* Skeleton, const FName& BoneName, const bool bIsVisible);
 	void OnBindingVisibilityStateChanged(UObject* BoundObject, const EBindingVisibilityState VisibilityState);
 	//called when added moved, deleted, added, all of which may need a refresh of a motion trail
 	void OnActorChangedSomehow(AActor* InActor);
 	void OnActorsChangedSomehow(TArray<AActor*>& InActors);	
 
+	FFrameNumber GetLocalTime() const;
+	
+	FGuid AddComponentToHierarchy(const FGuid& InBindingGuid, USceneComponent* CompToAdd, UMovieScene3DTransformTrack* TransformTrack);
+	FGuid AddControlRigTrail(USkeletalMeshComponent* Component, UControlRig* ControlRig, UMovieSceneControlRigParameterTrack* CRTrack, const FName& ControlName);
+	FGuid PinComponent(USceneComponent* InSceneComponent, FName InSocketName);
+
 private:
-	void AddControlRigTrail(USkeletalMeshComponent* Component,UControlRig* ControlRig, UMovieSceneControlRigParameterTrack* CRTrack, const FName& ControlName);
+	
 
-	void UpdateControlRig(const TArray<FFrameNumber>& Frames, UControlRig* ControlRig, TMap<FName, FGuid >& CompMapPair, bool bUseEditedTimes);
-	void UpdateControlRig(const FTrailEvaluateTimes& EvalTimes, UControlRig* ControlRig, TMap<FName, FGuid >& CompMapPair);
 
 	void UpdateSequencerBindings(const TArray<FGuid>& SequencerBindings, TFunctionRef<void(UObject*, FTrail*, FGuid)> OnUpdated);
 	void UpdateViewAndEvalRange();
 
-	void AddComponentToHierarchy(USceneComponent* CompToAdd, UMovieScene3DTransformTrack* TransformTrack);
-	void AddSkeletonToHierarchy(USkeletalMeshComponent* CompToAdd);
-
 	void RegisterControlRigDelegates(USkeletalMeshComponent* Component, UMovieSceneControlRigParameterTrack* CRParameterTrack);
 
 	TWeakPtr<ISequencer> WeakSequencer;
 	TMap<UObject*, FGuid> ObjectsTracked;
-	TMap<USkeletalMeshComponent*, TMap<FName, FGuid>> BonesTracked; 
-	TMap<UControlRig*, TMap<FName, FGuid>> ControlsTracked; 
+	TMap<USceneComponent*, TMap<FName, FGuid>> SocketsTracked;
+
+
+	struct FControlMapAndTransforms
+	{
+		TMap < FName, FGuid> NameToTrail;
+		TSharedPtr<UE::AIE::FArrayOfTransforms> ArrayOfTransforms;
+	};
+	TMap<UControlRig*, FControlMapAndTransforms> ControlsTracked;
 
 	// TODO: components can have multiple rigs so make this a map from sections to controls instead. However, this is only part of a larger problem of handling blending
 
@@ -103,7 +113,38 @@ private:
 		FDelegateHandle OnControlSelected;
 	};
 	TMap<UMovieSceneControlRigParameterTrack*, FControlRigDelegateHandles> ControlRigDelegateHandles;
+	
+private:
+	//new evaluate setup.. this is what we are evaluating...
+	TArray<UE::AIE::FActorAndWorldTransforms> EvaluatingActors;
+	TMap<UControlRig*, UE::AIE::FControlRigAndWorldTransforms> EvaluatingControlRigs;
+	TSet<FGuid> EvaluatingTrails;
+
+	//current frames
+	FCurrentFramesInfo  CurrentFramesInfo;
+	
+	void EvaluateSequencerAndSetTransforms();
+
+	static TArray<TUniquePtr <FTrail::FMotionTrailState >> PreviouslyPinnedTrails;
+public:
+	void PinTrail(FGuid InGuid);
+
+private:
+	//cached guid to check for updates
+	FGuid LastValidMovieSceneGuid;
+protected: //callbacks
+
+	void RegisterMotionTrailOptionDelegates();
+	void UnRegisterMotionTrailOptionDelegates();
+	void OnPinSelection();
+	void OnPinComponent(USceneComponent* InSceneComponent, FName SocketName);
+	void OnDeletePinned(FGuid InGuid);
+	void OnPutPinnedInSpace(FGuid InGuid, AActor* InActor, FName InComponentName);
+	void OnSetLinearColor(FGuid InGuid, FLinearColor Color);
+	void OnSetHasOffset(FGuid InGuid, bool bOffset);
 };
 
+
+
 } // namespace MovieScene
 } // namespace UE
diff --git a/Engine/Plugins/Editor/SequencerAnimTools/Source/SequencerAnimTools/Private/TrailHierarchy.cpp b/Engine/Plugins/Editor/SequencerAnimTools/Source/SequencerAnimTools/Private/TrailHierarchy.cpp
--- a/Engine/Plugins/Editor/SequencerAnimTools/Source/SequencerAnimTools/Private/TrailHierarchy.cpp
+++ b/Engine/Plugins/Editor/SequencerAnimTools/Source/SequencerAnimTools/Private/TrailHierarchy.cpp
@@ -15,6 +15,7 @@
 #include "Widgets/Input/SCheckBox.h"
 #include "Widgets/Text/STextBlock.h"
 #include "CanvasTypes.h"
+#include "Sequencer/SequencerTrailHierarchy.h"
 
 #define LOCTEXT_NAMESPACE "MotionTrailEditorMode"
 
@@ -22,61 +23,91 @@ namespace UE
 {
 namespace SequencerAnimTools
 {
+using namespace UE::AIE;
+
 
 IMPLEMENT_HIT_PROXY(HBaseTrailProxy, HHitProxy);
-IMPLEMENT_HIT_PROXY(HMotionTrailProxy, HBaseTrailProxy);
+IMPLEMENT_HIT_PROXY(HNewMotionTrailProxy, HBaseTrailProxy);
+
+
+void FTrail::GetColor(const FFrameNumber& CurrentTime, FColorState& State)
+{
+	if (State.GetStyle() == EMotionTrailTrailStyle::Time)
+	{
+		if (CurrentTime < State.SequencerTime)
+		{
+			State.CalculatedColor = State.Options->TimePreColor;
+		}
+		else
+		{
+			State.CalculatedColor = State.Options->TimePostColor;
+		}
+	}
+	else if (State.GetStyle() == EMotionTrailTrailStyle::Dashed)
+	{
+
+		if (State.bFirstFrame == true)
+		{
+			State.bFirstFrame = false;
+			State.StartFrame = CurrentTime;
+		}
+
+		const int32 Index = (CurrentTime.Value - State.StartFrame.Value) / State.TicksPerFrame.Value;
+
+		if (Index % 2 == 0)
+		{
+			State.CalculatedColor = State.Options->DashPreColor;
+		}
+		else
+		{
+			State.CalculatedColor = State.Options->DashPostColor;
+		}
+	}
+	else
+	{
+		if (State.bIsPinned && DrawInfo.IsValid())
+		{
+			State.CalculatedColor = DrawInfo->GetColor();
+		}
+		else
+		{
+			State.CalculatedColor = State.Options->DefaultColor;
+		}
+	}
+}
 
 void FTrailHierarchyRenderer::Render(const FSceneView* View, FPrimitiveDrawInterface* PDI)
 {
 	const FDateTime RenderStartTime = FDateTime::Now();
 	
-	const int32 NumEvalTimes = int32(OwningHierarchy->GetViewRange().Size<double>() / OwningHierarchy->GetSecondsPerSegment());
+	const int32 NumEvalTimes = int32((OwningHierarchy->GetViewFrameRange().GetUpperBoundValue().Value - OwningHierarchy->GetViewFrameRange().GetLowerBoundValue().Value) / OwningHierarchy->GetFramesPerFrame().Value) + 1;
 	const int32 NumLinesReserveSize = int32(NumEvalTimes * OwningHierarchy->GetAllTrails().Num() * 1.3);
 	PDI->AddReserveLines(SDPG_Foreground, NumLinesReserveSize);
 
 	TArray<FGuid> TrailsToRemove; //if trail is not visible we just get rid of it
 	// <parent, current>
+	FColorState ColorState;
+	ColorState.Setup(OwningHierarchy);
+	const FCurrentFramesInfo* CurrentFramesInfo = OwningHierarchy->GetCurrentFramesInfo();
 	for (const TPair<FGuid, TUniquePtr<FTrail>>& GuidTrailPair : OwningHierarchy->GetAllTrails())
 	{
 		FTrajectoryDrawInfo* CurDrawInfo = GuidTrailPair.Value->GetDrawInfo();
 		if (CurDrawInfo && OwningHierarchy->GetVisibilityManager().IsTrailVisible(GuidTrailPair.Key, GuidTrailPair.Value.Get()))
 		{
-			FDisplayContext DisplayContext = {
-				GuidTrailPair.Key,
-				FTrailScreenSpaceTransform(View),
-				UMotionTrailToolOptions::GetTrailOptions()->SecondsPerMark,
-				OwningHierarchy->GetViewRange(),
-				OwningHierarchy->GetSecondsPerSegment()
-			};
-
-			const bool bHitTesting = PDI && PDI->IsHitTesting();
-			TArray<FVector> PointsToDraw;
-			TArray<double> TrajectoryTimes;
-			GuidTrailPair.Value->GetTrajectoryPointsForDisplay(DisplayContext, PointsToDraw, TrajectoryTimes);
-			FLinearColor Color = OwningHierarchy->GetVisibilityManager().IsTrailAlwaysVisible(GuidTrailPair.Key) ? CurDrawInfo->GetColor() : UMotionTrailToolOptions::GetTrailOptions()->TrailColor;
-			Color = GuidTrailPair.Value->IsTrailSelected() ? FLinearColor(1.0f, 1.0f, 0.0f) : Color;
-			if (PointsToDraw.Num() > 1)
+			if (OwningHierarchy->IsTrailEvaluating(GuidTrailPair.Key, true))
 			{
-				const float TrailThickness = UMotionTrailToolOptions::GetTrailOptions()->TrailThickness;
-				FVector LastPoint = PointsToDraw[0];
-				for (int32 Idx = 1; Idx < PointsToDraw.Num(); Idx++)
-				{
-					if (bHitTesting)
-					{
-						const double CurSecond = TrajectoryTimes[Idx - 1];
-						PDI->SetHitProxy(new HMotionTrailProxy(GuidTrailPair.Key, LastPoint, CurSecond));
-					}
-					const FVector CurPoint = PointsToDraw[Idx];
-					PDI->DrawLine(LastPoint, CurPoint, Color, SDPG_Foreground,TrailThickness);
-					LastPoint = CurPoint;
-					if (bHitTesting)
-					{
-						PDI->SetHitProxy(nullptr);
-					}
-				}
+				GuidTrailPair.Value->RenderEvaluating(GuidTrailPair.Key, View, PDI);
+				continue;
 			}
+			if (GuidTrailPair.Value.Get()->GetCacheState() != ETrailCacheState::UpToDate)
+			{
+				continue;
+			}
+			//the above check so that if we are evaluating something that is getting drawn indirectly we skip drawing it, but if it is
+			//getting manipulated we still want to draw it. If it's not evaluating though we don't use sparse time values, instead do every frame.
+			const bool bTrailDirectlyEvaluating = OwningHierarchy->IsTrailEvaluating(GuidTrailPair.Key, false);
+			GuidTrailPair.Value->ReadyToDrawTrail(ColorState, CurrentFramesInfo, bTrailDirectlyEvaluating,  OwningHierarchy->GetVisibilityManager().IsTrailAlwaysVisible(GuidTrailPair.Key));
 			GuidTrailPair.Value->Render(GuidTrailPair.Key,View,PDI);
-
 		}
 		else
 		{
@@ -94,54 +125,49 @@ void FTrailHierarchyRenderer::Render(const FSceneView* View, FPrimitiveDrawInter
 
 void FTrailHierarchyRenderer::DrawHUD(const FSceneView* View, FCanvas* Canvas)
 {
-
-	if (UMotionTrailToolOptions::GetTrailOptions()->bShowMarks == false)
+	if (UMotionTrailToolOptions::GetTrailOptions()->bShowMarks)
 	{
-		for (const TPair<FGuid, TUniquePtr<FTrail>>& GuidTrailPair : OwningHierarchy->GetAllTrails())
-		{
-			GuidTrailPair.Value->DrawHUD( View, Canvas);
-		}
-		return;
-	}	
+		const FDateTime DrawHUDStartTime = FDateTime::Now();
+		const FFrameNumber FramesPerMark = OwningHierarchy->GetFramesPerFrame();
 
-	const FDateTime DrawHUDStartTime = FDateTime::Now();
+		const int32 NumMarksPerTrail = int32((OwningHierarchy->GetViewFrameRange().GetUpperBoundValue().Value - OwningHierarchy->GetViewFrameRange().GetLowerBoundValue().Value) / FramesPerMark.Value) + 1;
+		const int32 PredictedNumMarks = int32(NumMarksPerTrail * OwningHierarchy->GetAllTrails().Num() * 1.3);
 
-	const double SecondsPerMark = UMotionTrailToolOptions::GetTrailOptions()->bLockMarksToFrames ? OwningHierarchy->GetSecondsPerFrame() : UMotionTrailToolOptions::GetTrailOptions()->SecondsPerMark;
-	const int32 PredictedNumMarks = int32((OwningHierarchy->GetViewRange().Size<double>() / SecondsPerMark) * OwningHierarchy->GetAllTrails().Num() * 1.3); // Multiply by 1.3 to be safe
-	Canvas->GetBatchedElements(FCanvas::EElementType::ET_Line)->AddReserveLines(PredictedNumMarks);
+		Canvas->GetBatchedElements(FCanvas::EElementType::ET_Line)->AddReserveLines(PredictedNumMarks);
+		FTrailScreenSpaceTransform Transform(View, Canvas->GetDPIScale());
+		const FCurrentFramesInfo* CurrentFramesInfo = OwningHierarchy->GetCurrentFramesInfo();
 
-	for(const TPair<FGuid, TUniquePtr<FTrail>>& GuidTrailPair : OwningHierarchy->GetAllTrails())
-	{
-		FTrajectoryDrawInfo* CurDrawInfo = GuidTrailPair.Value->GetDrawInfo();
-		if (CurDrawInfo && OwningHierarchy->GetVisibilityManager().IsTrailVisible(GuidTrailPair.Key, GuidTrailPair.Value.Get()))
+		for (const TPair<FGuid, TUniquePtr<FTrail>>& GuidTrailPair : OwningHierarchy->GetAllTrails())
 		{
-			FDisplayContext DisplayContext = {
-				GuidTrailPair.Key,
-				FTrailScreenSpaceTransform(View, Canvas->GetDPIScale()),
-				SecondsPerMark,
-				OwningHierarchy->GetViewRange(),
-				OwningHierarchy->GetSecondsPerSegment()
-
-			};
-
-			TArray<FVector2D> Marks, MarkNormals;
-			GuidTrailPair.Value->GetTickPointsForDisplay(DisplayContext, Marks, MarkNormals);
-			const FLinearColor Color = OwningHierarchy->GetVisibilityManager().IsTrailAlwaysVisible(GuidTrailPair.Key) ? CurDrawInfo->GetColor() : UMotionTrailToolOptions::GetTrailOptions()->MarkColor;
-			for (int32 Idx = 0; Idx < Marks.Num(); Idx++)
+			if (OwningHierarchy->GetVisibilityManager().IsTrailVisible(GuidTrailPair.Key, GuidTrailPair.Value.Get()))
 			{
-				const FVector2D StartPoint = Marks[Idx] - MarkNormals[Idx] * UMotionTrailToolOptions::GetTrailOptions()->MarkSize;
-				const FVector2D EndPoint = Marks[Idx] + MarkNormals[Idx] * UMotionTrailToolOptions::GetTrailOptions()->MarkSize;
-				FCanvasLineItem LineItem = FCanvasLineItem(StartPoint, EndPoint);
-				LineItem.SetColor(Color);
-				Canvas->DrawItem(LineItem);
+				if (OwningHierarchy->IsTrailEvaluating(GuidTrailPair.Key, true) || GuidTrailPair.Value.Get()->GetCacheState() != ETrailCacheState::UpToDate)
+				{
+					continue;
+				}
+				TArray<FVector2D> Marks, MarkNormals;
+				const bool bTrailDirectlyEvaluating = OwningHierarchy->IsTrailEvaluating(GuidTrailPair.Key, false);
+				GuidTrailPair.Value->GetTickPointsForDisplay(Transform, *CurrentFramesInfo, bTrailDirectlyEvaluating, Marks, MarkNormals);
+				if (Marks.Num() > 0)
+				{
+					const FLinearColor Color = UMotionTrailToolOptions::GetTrailOptions()->MarkColor;
+					for (int32 Idx = 0; Idx < Marks.Num(); Idx++)
+					{
+						const FVector2D StartPoint = Marks[Idx] - MarkNormals[Idx] * UMotionTrailToolOptions::GetTrailOptions()->MarkSize;
+						const FVector2D EndPoint = Marks[Idx] + MarkNormals[Idx] * UMotionTrailToolOptions::GetTrailOptions()->MarkSize;
+						FCanvasLineItem LineItem = FCanvasLineItem(StartPoint, EndPoint);
+						LineItem.SetColor(Color);
+						Canvas->DrawItem(LineItem);
+					}
+				}
 			}
+			GuidTrailPair.Value->DrawHUD(View, Canvas);
+
 		}
-		GuidTrailPair.Value->DrawHUD(View, Canvas);
 
+		const FTimespan DrawHUDTimespan = FDateTime::Now() - DrawHUDStartTime;
+		OwningHierarchy->GetTimingStats().Add("FTrailHierarchyRenderer::DrawHUD", DrawHUDTimespan);
 	}
-
-	const FTimespan DrawHUDTimespan = FDateTime::Now() - DrawHUDStartTime;
-	OwningHierarchy->GetTimingStats().Add("FTrailHierarchyRenderer::DrawHUD", DrawHUDTimespan);
 }
 
 class STrailOptionsPopup : public SCompoundWidget
@@ -153,7 +179,6 @@ public:
 
 	SLATE_ARGUMENT(FTrailVisibilityManager*, InVisibilityManager)
 	SLATE_ARGUMENT(FGuid, InTrailGuid)
-	SLATE_ARGUMENT(FTrajectoryDrawInfo*, InDrawInfo)
 
 	SLATE_END_ARGS()
 
@@ -161,21 +186,32 @@ public:
 public:
 	FTrailVisibilityManager* VisibilityManager;
 	FGuid TrailGuid;
-	FTrajectoryDrawInfo* DrawInfo;
 
-	void SetTrailAlwaysVisible(ECheckBoxState NewState)
+	void SetPinned(ECheckBoxState NewState)
 	{
-		bool bSet = (NewState == ECheckBoxState::Checked);
-		if (VisibilityManager)
+		UMotionTrailToolOptions* Settings = UMotionTrailToolOptions::GetTrailOptions();
+		if (NewState == ECheckBoxState::Checked)
 		{
-			VisibilityManager->SetTrailAlwaysVisible(TrailGuid,bSet);
-			if (DrawInfo && bSet == true)
+			Settings->PinSelection();
+		}
+		else
+		{
+			if (Settings->GetNumPinned() > 0)
 			{
-				DrawInfo->SetColor(UMotionTrailToolOptions::GetTrailOptions()->TrailColor);
+				for (int32 Index = 0; Index < Settings->GetNumPinned(); ++Index)
+				{
+					if (UMotionTrailToolOptions::FPinnedTrail* Trail = Settings->GetPinnedTrail(Index))
+					{
+						if (Trail->TrailGuid == TrailGuid)
+						{
+							Settings->DeletePinned(Index);
+							break;
+						}
+					}
+				}
 			}
 		}
 	}
-
 	ECheckBoxState IsTrailAlwaysVisible() const
 	{
 		if (VisibilityManager)
@@ -196,7 +232,6 @@ public:
 	{
 		VisibilityManager = InArgs._InVisibilityManager;
 		TrailGuid = InArgs._InTrailGuid;
-		DrawInfo = InArgs._InDrawInfo;
 		// Then make widget
 		this->ChildSlot
 		[
@@ -216,7 +251,7 @@ public:
 					.VAlign(VAlign_Center)
 					[
 						SNew(SCheckBox)
-						.OnCheckStateChanged(this, &STrailOptionsPopup::SetTrailAlwaysVisible)
+						.OnCheckStateChanged(this, &STrailOptionsPopup::SetPinned)
 						.IsChecked(this, &STrailOptionsPopup::IsTrailAlwaysVisible)
 					]
 
@@ -226,18 +261,9 @@ public:
 					[
 						SNew(STextBlock)
 						.AutoWrapText(true)
-						.Text(LOCTEXT("AlwaysVisible", "Always Visible"))
+						.Text(LOCTEXT("Pinned", "Pinned"))
 					]	
 				]
-				+ SVerticalBox::Slot()
-				.AutoHeight()
-				.Padding(0.0f, 1.0f)
-				[
-					SNew(SButton)
-					.ContentPadding(FMargin(8, 2))
-					.Text(LOCTEXT("Show Options", "Show Options"))
-					.OnClicked(this, &STrailOptionsPopup::ShowOptions)
-				]
 			]
 		];
 		
@@ -253,13 +279,10 @@ void FTrailHierarchy::OpenContextMenu(const FGuid& TrailGuid)
 	TSharedPtr<SWidget> MenuWidget;
 	if(const TUniquePtr<FTrail>* Trail = GetAllTrails().Find(TrailGuid))
 	{
-		FTrajectoryDrawInfo* CurDrawInfo = (*Trail)->GetDrawInfo();
-
 		MenuWidget =
 			SNew(STrailOptionsPopup)
 			.InVisibilityManager(&GetVisibilityManager())
-			.InTrailGuid(TrailGuid)
-			.InDrawInfo(CurDrawInfo);
+			.InTrailGuid(TrailGuid);
 
 		// Create as context menu
 		FSlateApplication::Get().PushMenu(
@@ -289,7 +312,7 @@ bool FTrailHierarchy::HandleClick(FEditorViewportClient* InViewportClient, HHitP
 {
 	bool bClickIsHandled = false;
 	
-	if (HMotionTrailProxy* HitProxy = HitProxyCast<HMotionTrailProxy>(InHitProxy))
+	if (HNewMotionTrailProxy* HitProxy = HitProxyCast<HNewMotionTrailProxy>(InHitProxy))
 	{
 		if (Click.bIsRightMouse)
 		{
@@ -346,16 +369,23 @@ bool FTrailHierarchy::IsAnythingSelected(FVector& OutVectorPosition) const
 	return (NumSelected > 0);
 }
 
-bool FTrailHierarchy::IsAnythingSelected(TArray<FVector>& OutVectorPositions) const
+bool FTrailHierarchy::IsAnythingSelected(TArray<FVector>& OutVectorPositions, bool bAllPositions) const
 {
 	OutVectorPositions.SetNum(0);
 	for (const TPair<FGuid, TUniquePtr<FTrail>>& GuidTrailPair : GetAllTrails())
 	{
-		FVector Location;
-		bool bHandled = GuidTrailPair.Value->IsAnythingSelected(Location);
-		if (bHandled)
+		if (bAllPositions == false)
 		{
-			OutVectorPositions.Add(Location);
+			FVector Location;
+			bool bHandled = GuidTrailPair.Value->IsAnythingSelected(Location);
+			if (bHandled)
+			{
+				OutVectorPositions.Add(Location);
+			}
+		}
+		else
+		{
+			GuidTrailPair.Value->IsAnythingSelected(OutVectorPositions);
 		}
 	}
 	return (OutVectorPositions.Num() > 0);
@@ -394,6 +424,15 @@ void FTrailHierarchy::DeleteSelectedKeys()
 	}
 }
 
+bool FTrailHierarchy::IsVisible(const FGuid& Key) const
+{
+	if (UMotionTrailToolOptions::GetTrailOptions()->bShowSelectedTrails == true)
+	{
+		return AllTrails.Contains(Key);
+	}
+	else return VisibilityManager.IsTrailAlwaysVisible(Key);
+}
+
 bool FTrailHierarchy::IsAlwaysVisible(const FGuid Key) const
 {
 	return VisibilityManager.IsTrailAlwaysVisible(Key);
@@ -441,13 +480,12 @@ bool FTrailHierarchy::StartTracking()
 	return bIsHandled;
 }
 
-bool FTrailHierarchy::ApplyDelta(const FVector& Pos, const FRotator& Rot, const FVector &WidgetLocation)
-
+bool FTrailHierarchy::ApplyDelta(const FVector& Pos, const FRotator& Rot, const FVector &WidgetLocation, bool bApplyToOffset)
 {
 	bool bIsHandled = false;
 	for (const TPair<FGuid, TUniquePtr<FTrail>>& GuidTrailPair : GetAllTrails())
 	{
-		bool bHandled =  GuidTrailPair.Value->ApplyDelta(Pos,Rot,WidgetLocation);
+		bool bHandled =  GuidTrailPair.Value->ApplyDelta(Pos,Rot,WidgetLocation,bApplyToOffset);
 		if (bHandled)
 		{
 			bIsHandled = true;
@@ -470,66 +508,65 @@ bool FTrailHierarchy::EndTracking()
 	return bIsHandled;
 }
 
+
 void FTrailHierarchy::CalculateEvalRangeArray()
 {
-	// Generate times to evaluate
-	SecondsPerSegment = GetSecondsPerSegment();
-
-	EvalTimesArr.Reserve(int32(EvalRange.Size<double>() / SecondsPerSegment) + 1);
-	EvalTimesArr.SetNum(0);
-	for (double SecondsItr = EvalRange.GetLowerBoundValue(); SecondsItr <= EvalRange.GetUpperBoundValue() + SecondsPerSegment; SecondsItr += SecondsPerSegment)
-	{
-		EvalTimesArr.Add(SecondsItr);
-	}
+	// new
+	TicksPerSegment = GetFramesPerFrame();
 
-	if (LastSecondsPerSegment != SecondsPerSegment || EvalRange != LastEvalRange)
+	if (LastTicksPerSegment != TicksPerSegment || TickEvalRange != LastTickEvalRange)
 	{
 		for (const TPair<FGuid, TUniquePtr<FTrail>>& GuidTrailPair : GetAllTrails())
 		{
 			GuidTrailPair.Value->ForceEvaluateNextTick();
 		}
-		LastSecondsPerSegment = SecondsPerSegment;
-		LastEvalRange = EvalRange;
+		LastTicksPerSegment = TicksPerSegment;
+		LastTickEvalRange = TickEvalRange;
 	}
 }
 void FTrailHierarchy::Update()
 {
+	//new
 	const FDateTime UpdateStartTime = FDateTime::Now();
 	
 	CalculateEvalRangeArray();
-	FTrailEvaluateTimes EvalTimes = FTrailEvaluateTimes(EvalTimesArr, SecondsPerSegment);
-
 	VisibilityManager.InactiveMask.Reset();
 	TArray<FGuid> DeadTrails;
 
+	const bool bCheckForChange = CheckForChanges();
+
 	for (const TPair<FGuid, TUniquePtr<FTrail>>& GuidTrailPair : GetAllTrails())
 	{
-		FGuid CurGuid = GuidTrailPair.Key;
 
-		FTrail::FSceneContext SceneContext = {
-			CurGuid,
-			EvalTimes,
-			this
+		const FGuid CurGuid = GuidTrailPair.Key;
+		FTrail::FNewSceneContext SceneContext = {
+				bCheckForChange,
+				CurGuid,
+				this
 		};
-
 		if (!AllTrails.Contains(CurGuid))
 		{
 			continue;
 		}
-
+		//don't update if not visible
+		if (IsVisible(GuidTrailPair.Key) == false)
+		{
+			continue;
+		}
 		// Update the trail
-		ETrailCacheState CurCacheState = AllTrails[CurGuid]->UpdateTrail(SceneContext);
-		if (CurCacheState == ETrailCacheState::Dead)
+		FTrailCurrentStatus Status = AllTrails[CurGuid]->UpdateTrail(SceneContext);
+
+		if (Status.CacheState == ETrailCacheState::Dead)
 		{
 			DeadTrails.Add(CurGuid);
 		}
-
-		if (CurCacheState == ETrailCacheState::NotUpdated)
+		if (Status.CacheState == ETrailCacheState::NotUpdated)
 		{
 			VisibilityManager.InactiveMask.Add(CurGuid);
 		}
 	}
 
+	//remove dead trails
 	for (const FGuid& TrailGuid : DeadTrails)
 	{
 		RemoveTrail(TrailGuid);
@@ -537,6 +574,7 @@ void FTrailHierarchy::Update()
 
 	const FTimespan UpdateTimespan = FDateTime::Now() - UpdateStartTime;
 	TimingStats.Add("FTrailHierarchy::Update", UpdateTimespan);
+
 }
 
 void FTrailHierarchy::AddTrail(const FGuid& Key, TUniquePtr<FTrail>&& TrailPtr)
diff --git a/Engine/Plugins/Editor/SequencerAnimTools/Source/SequencerAnimTools/Private/TrajectoryDrawInfo.cpp b/Engine/Plugins/Editor/SequencerAnimTools/Source/SequencerAnimTools/Private/TrajectoryDrawInfo.cpp
--- a/Engine/Plugins/Editor/SequencerAnimTools/Source/SequencerAnimTools/Private/TrajectoryDrawInfo.cpp
+++ b/Engine/Plugins/Editor/SequencerAnimTools/Source/SequencerAnimTools/Private/TrajectoryDrawInfo.cpp
@@ -3,12 +3,26 @@
 #include "TrajectoryDrawInfo.h"
 #include "SceneView.h"
 #include "TrajectoryCache.h"
+#include "Framework/Application/SlateApplication.h"
+
+
+int32 UE::SequencerAnimTools::FFrameCalculator::IndicesToCalculatePerBucket = 10;
+static FAutoConsoleVariableRef CVarIndicesToCalculatePerBucket(
+	TEXT("Sequencer.MotionTrailTickIterations"),
+	UE::SequencerAnimTools::FFrameCalculator::IndicesToCalculatePerBucket,
+	TEXT("Number of iterations to tick sequencer when calculating motion trails"));
 
 namespace UE
 {
 namespace SequencerAnimTools
 {
 
+/**********************************************************************************
+*
+*    FTrailspaceTransform
+*
+**********************************************************************************/
+
 TOptional<FVector2D> FTrailScreenSpaceTransform::ProjectPoint(const FVector& Point) const
 {
 
@@ -22,37 +36,242 @@ TOptional<FVector2D> FTrailScreenSpaceTransform::ProjectPoint(const FVector& Poi
 	return TOptional<FVector2D>();
 }
 
-FVector FTrajectoryDrawInfo::GetPoint(double InTime)
+/**********************************************************************************
+*
+*    FFrameCalculator
+*
+**********************************************************************************/
+
+void FFrameCalculator::SetUpFrameCalculator(const UE::AIE::FFrameTimeByIndex& InCurrentFrameTimes, TRange<FFrameNumber> ViewRange)
 {
-	FTransform Transform = TrajectoryCache->GetInterp(InTime);
-	return Transform.GetLocation();
+	CurrentFrameTimes = InCurrentFrameTimes;
+	//first makes sure viewrange is within eval range.
+	if (ViewRange.GetLowerBoundValue() < InCurrentFrameTimes.StartFrame)
+	{
+		ViewRange.SetLowerBoundValue(InCurrentFrameTimes.StartFrame);
+	}
+	if (ViewRange.GetUpperBoundValue() > InCurrentFrameTimes.EndFrame)
+	{
+		ViewRange.SetUpperBoundValue(InCurrentFrameTimes.EndFrame);
+	}
+
+	TRange<FFrameNumber> Range = ViewRange;
+	Ranges.Add(Range);
+	CurrentRangeStartIndex = 0;
+	CurrentRangeEndIndex = InCurrentFrameTimes.NumFrames - 1;
+
+	Reset();
+	CalculateStartEndIndices();
+	NumOfBuckets = FMath::RoundToInt32((float)(CurrentRangeEndIndex + 1 - CurrentRangeStartIndex) / (float)(IndicesToCalculatePerBucket));
 }
 
-void FTrajectoryDrawInfo::GetTrajectoryPointsForDisplay(const FDisplayContext& InDisplayContext, TArray<FVector>& TrajectoryPoints, TArray<double>& TrajectoryTimes)
-{	
-	CachedViewRange = InDisplayContext.TimeRange;
+void FFrameCalculator::CalculateStartEndIndices()
+{
+	if (CurrentRange < Ranges.Num())
+	{
+		TRange<FFrameNumber>& Range = Ranges[CurrentRange];
+		CurrentRangeStartIndex = CurrentFrameTimes.CalculateIndex(Range.GetLowerBoundValue());
+		CurrentRangeEndIndex = CurrentFrameTimes.CalculateIndex(Range.GetUpperBoundValue());
+	}
 
-	TrajectoryTimes = TrajectoryCache->GetAllTimesInRange(InDisplayContext.TimeRange);
-	TrajectoryPoints.Reserve(TrajectoryTimes.Num());
-	
-	for (const double Time : TrajectoryTimes)
+}
+void FFrameCalculator::AddMustHaveIndices(const TArray<int32>& InMustHaveIndices)
+{
+	MustHaveIndices = InMustHaveIndices;
+}
+
+void FFrameCalculator::Reset()
+{
+	CurrentRange = 0;
+	CurrentBucket = -1; //-1 bucket is just MustHaveTimes
+	IndicesToCalculate.SetNum(0);
+}
+
+//calculate next set of indices in the IndicesToCalculate array
+//will return true if we need to calculate, will return false if this is the last time
+bool FFrameCalculator::CalculateIndices()
+{
+	IndicesToCalculate.SetNum(0);
+	//any mouse press just do base bucket
+	if (FSlateApplication::Get().GetPressedMouseButtons().Contains(EKeys::LeftMouseButton) == true		||
+		FSlateApplication::Get().GetPressedMouseButtons().Contains(EKeys::MiddleMouseButton) == true	||
+		FSlateApplication::Get().GetPressedMouseButtons().Contains(EKeys::RightMouseButton) == true
+		)
+	{
+		CurrentBucket = -1;
+		CurrentRange = 0;
+	}
+	if (CurrentBucket >= 0)
+	{
+		for (int32 Index = 0; Index < IndicesToCalculatePerBucket; ++Index)
+		{
+			const int32 CalculatedIndex = CurrentRangeStartIndex + CurrentBucket + (Index * NumOfBuckets);
+			if (CalculatedIndex <= CurrentRangeEndIndex) //can go over
+			{
+				IndicesToCalculate.Add(CalculatedIndex);
+			}
+		}
+	}
+	else
+	{
+		IndicesToCalculate = MustHaveIndices;
+	}
+	++CurrentBucket;
+	if (CurrentBucket > NumOfBuckets)
 	{
-		TrajectoryPoints.Add(TrajectoryCache->Get(Time).GetTranslation());
+		++CurrentRange;
+		if (CurrentRange >= Ranges.Num())
+		{
+			return false;
+		}
+		CalculateStartEndIndices();
 	}
+	return true;
+}
+/**********************************************************************************
+*
+*    FCurrentFramesInfo
+*
+**********************************************************************************/
 
+void FCurrentFramesInfo::SetViewRange(const TRange<FFrameNumber>& InViewRange)
+{
+	ViewRange = InViewRange;
+	bViewRangeIsEvalRange = (CurrentFrameTimes.StartFrame == ViewRange.GetLowerBoundValue() &&
+		CurrentFrameTimes.EndFrame == ViewRange.GetUpperBoundValue());
+}
+void FCurrentFramesInfo::SetUpFrameTimes(const TRange<FFrameNumber>& InEvalFrameRange, const FFrameNumber& InFrameStep)
+{
+	CurrentFrameTimes = UE::AIE::FFrameTimeByIndex(InEvalFrameRange.GetLowerBoundValue(), InEvalFrameRange.GetUpperBoundValue(), InFrameStep);
+	SortedTransformIndices.Reserve(CurrentFrameTimes.NumFrames);
+	TransformIndices.Reserve(CurrentFrameTimes.NumFrames);
+	TransformIndices.Reset();
+
+	FrameCalculator.SetUpFrameCalculator(CurrentFrameTimes, ViewRange);
 }
 
-void FTrajectoryDrawInfo::GetTickPointsForDisplay(const FDisplayContext& InDisplayContext, TArray<FVector2D>& Ticks, TArray<FVector2D>& TickNormals)
+void FCurrentFramesInfo::AddMustHaveTimes(const TSet<FFrameNumber>& InMustHaveTimes)
 {
-	const double FirstTick = FMath::FloorToDouble(InDisplayContext.TimeRange.GetLowerBoundValue() / InDisplayContext.SecondsPerTick) * InDisplayContext.SecondsPerTick;
-	const double Spacing = InDisplayContext.SecondsPerSegment;
+	TArray<int32> MustHaveIndices;
+	for (const FFrameNumber& FrameNumber : InMustHaveTimes)
+	{
+		int32 Index = CurrentFrameTimes.CalculateIndex(FrameNumber);
+		MustHaveIndices.Add(Index);
+	}
+	FrameCalculator.AddMustHaveIndices(MustHaveIndices);
+}
+
+const TArray<int32>& FCurrentFramesInfo::IndicesToCalculate() const
+{
+	return FrameCalculator.IndicesToCalculate;
+}
+
+bool FCurrentFramesInfo::KeepCalculating()
+{
+	const bool bKeep = FrameCalculator.CalculateIndices();
+	for (int32 Index : FrameCalculator.IndicesToCalculate)
+	{
+		const FFrameNumber FrameNumber = CurrentFrameTimes.CalculateFrame(Index);
+		SortedTransformIndices.Add(Index, FrameNumber);
+	}
+	SortedTransformIndices.GenerateKeyArray(TransformIndices);
+	SortedTransformIndices.GenerateValueArray(CurrentFrames);
+	return bKeep;
+}
+
+void FCurrentFramesInfo::Reset()
+{
+	FrameCalculator.Reset();
+	SortedTransformIndices.Reset();
+	TransformIndices.Reset();
+	CurrentFrames.SetNum(0);
+}
+
+/**********************************************************************************
+*
+*    FTrajectoryDrawInfo
+*
+**********************************************************************************/
+void FTrajectoryDrawInfo::GetTrajectoryPointsForDisplay(const FTransform& OffsetTransform, const FTransform& ParentSpaceTransform, const FCurrentFramesInfo& InCurrentFramesInfo, bool bIsEvaluating,  TArray<FVector>& OutPoints, TArray<FFrameNumber>&  OutFrames)
+{
+	//if not evaluating we just sample the times that are in range
+	//if we are evaluating then we use the InCurrentFramesInfo.TransformIndices to get the sparse values
+	if (bIsEvaluating == false )
+	{
+		OutPoints.Reserve(InCurrentFramesInfo.TransformIndices.Num());
+		OutPoints.SetNum(0);
+		OutFrames.Reserve(InCurrentFramesInfo.TransformIndices.Num());
+		OutFrames.SetNum(0);
+		const int32 StartIndex = InCurrentFramesInfo.CurrentFrameTimes.CalculateIndex(InCurrentFramesInfo.ViewRange.GetLowerBoundValue());
+		const int32 EndIndex = InCurrentFramesInfo.CurrentFrameTimes.CalculateIndex(InCurrentFramesInfo.ViewRange.GetUpperBoundValue());
+		for (int64 Index = StartIndex; Index <= EndIndex; ++Index)
+		{
+			const FFrameNumber Frame = InCurrentFramesInfo.CurrentFrameTimes.CalculateFrame(Index);
+			FVector Position = ParentSpaceTransform.TransformPosition(ArrayOfTransforms->Transforms[Index].GetLocation());
+			Position = OffsetTransform.TransformPosition(Position);
+			OutPoints.Add(Position);
+			OutFrames.Add(Frame);
+		}
+	}
+	else //if (bIsEvaluating)
+	{
+		if (InCurrentFramesInfo.bViewRangeIsEvalRange)
+		{
+			OutPoints.SetNum(InCurrentFramesInfo.TransformIndices.Num());
+			OutFrames.SetNum(InCurrentFramesInfo.TransformIndices.Num());
+			for (int32 Index = 0; Index < InCurrentFramesInfo.TransformIndices.Num(); ++Index)
+			{
+				const int32 TransformIndex = InCurrentFramesInfo.TransformIndices[Index];
+				FVector Position = ParentSpaceTransform.TransformPosition(ArrayOfTransforms->Transforms[TransformIndex].GetLocation());
+				Position = OffsetTransform.TransformPosition(Position);
+				OutPoints[Index] = Position;
+				OutFrames[Index] = InCurrentFramesInfo.CurrentFrames[Index];
+			}
+		}
+		else
+		{
+			OutPoints.Reserve(InCurrentFramesInfo.TransformIndices.Num());
+			OutPoints.SetNum(0);
+			OutFrames.Reserve(InCurrentFramesInfo.TransformIndices.Num());
+			OutFrames.SetNum(0);
+			const int32 StartIndex = InCurrentFramesInfo.CurrentFrameTimes.CalculateIndex(InCurrentFramesInfo.ViewRange.GetLowerBoundValue());
+			const int32 EndIndex = InCurrentFramesInfo.CurrentFrameTimes.CalculateIndex(InCurrentFramesInfo.ViewRange.GetUpperBoundValue());
+			for (int32 Index = 0; Index < InCurrentFramesInfo.TransformIndices.Num(); ++Index)
+			{
+				const int32 TransformIndex = InCurrentFramesInfo.TransformIndices[Index];
+				if (TransformIndex >= StartIndex && TransformIndex <= EndIndex)
+				{
+					FVector Position = ParentSpaceTransform.TransformPosition(ArrayOfTransforms->Transforms[TransformIndex].GetLocation());
+					Position = OffsetTransform.TransformPosition(Position);
+					OutPoints.Add(Position);
+					OutFrames.Add(InCurrentFramesInfo.CurrentFrames[Index]);
+
+				}
+				else if (TransformIndex > EndIndex)
+				{
+					break;
+				}
+			}
+		}
+	}
+}
 
-	for (double TickItr = FirstTick + InDisplayContext.SecondsPerTick; TickItr < InDisplayContext.TimeRange.GetUpperBoundValue(); TickItr += InDisplayContext.SecondsPerTick)
+void FTrajectoryDrawInfo::GetTickPointsForDisplay(const FTransform& OffsetTransform, const FTransform& ParentSpaceTransform, const FTrailScreenSpaceTransform& InScreenSpaceTransform, const FCurrentFramesInfo& InCurrentFramesInfo, bool bIsEvaluating, TArray<FVector2D>& Ticks, TArray<FVector2D>& TickNormals)
+{
+	if (ArrayOfTransforms.IsValid() == false || ArrayOfTransforms->Transforms.Num() == 0 || InCurrentFramesInfo.TransformIndices.Num() != InCurrentFramesInfo.CurrentFrames.Num())
+	{
+		return;
+	}
+	const FVector PrevInterpolated = GetPoint(OffsetTransform, ParentSpaceTransform, InCurrentFramesInfo, InCurrentFramesInfo.CurrentFrameTimes.StartFrame);
+	TOptional<FVector2D> PrevProjected = InScreenSpaceTransform.ProjectPoint(PrevInterpolated);
+	const int32 StartIndex = InCurrentFramesInfo.CurrentFrameTimes.CalculateIndex(InCurrentFramesInfo.ViewRange.GetLowerBoundValue());
+	const int32 EndIndex = InCurrentFramesInfo.CurrentFrameTimes.CalculateIndex(InCurrentFramesInfo.ViewRange.GetUpperBoundValue());
+	for (int64 Index = StartIndex; Index < EndIndex; ++Index)
 	{
-		FVector Interpolated = TrajectoryCache->GetInterp(TickItr).GetTranslation();
-		TOptional<FVector2D> Projected = InDisplayContext.ScreenSpaceTransform.ProjectPoint(Interpolated);
-		const double PrevSampleTime = TickItr + Spacing < InDisplayContext.TimeRange.GetUpperBoundValue() ? TickItr + Spacing : TickItr - Spacing;
-		TOptional<FVector2D> PrevProjected = InDisplayContext.ScreenSpaceTransform.ProjectPoint(TrajectoryCache->GetInterp(PrevSampleTime).GetTranslation());
+		const FFrameNumber Frame = InCurrentFramesInfo.CurrentFrameTimes.CalculateFrame(Index);
+		FVector Interpolated = GetPoint(OffsetTransform, ParentSpaceTransform, InCurrentFramesInfo, Frame);
+
+		TOptional<FVector2D> Projected = InScreenSpaceTransform.ProjectPoint(Interpolated);
 
 		if (Projected && PrevProjected)
 		{
@@ -63,8 +282,19 @@ void FTrajectoryDrawInfo::GetTickPointsForDisplay(const FDisplayContext& InDispl
 
 			TickNormals.Add(FVector2D(-Diff.Y, Diff.X));
 		}
+		PrevProjected = Projected;
 	}
 }
 
+FVector FTrajectoryDrawInfo::GetPoint(const FTransform& OffsetTransform, const FTransform& ParentSpaceTransform, const FCurrentFramesInfo& InCurrentFramesInfo, const FFrameNumber& InTime)
+{
+	FTransform Transform = ArrayOfTransforms->Interp(InTime, InCurrentFramesInfo.TransformIndices, InCurrentFramesInfo.CurrentFrames);
+	FVector Position = ParentSpaceTransform.TransformPosition(Transform.GetLocation());
+	Position = OffsetTransform.TransformPosition(Position);
+
+	return Position;
+}
+
 } // namespace SequencerAnimTools
 } // namespace UE
+
diff --git a/Engine/Plugins/Editor/SequencerAnimTools/Source/SequencerAnimTools/Public/MotionTrailTool.h b/Engine/Plugins/Editor/SequencerAnimTools/Source/SequencerAnimTools/Public/MotionTrailTool.h
--- a/Engine/Plugins/Editor/SequencerAnimTools/Source/SequencerAnimTools/Public/MotionTrailTool.h
+++ b/Engine/Plugins/Editor/SequencerAnimTools/Source/SequencerAnimTools/Public/MotionTrailTool.h
@@ -110,6 +110,9 @@ public:
 	/** Frame the selection*/
 	TSharedPtr<FUICommandInfo> FrameSelection;
 
+	/** Deselect All*/
+	TSharedPtr<FUICommandInfo> DeselectAll;
+
 	TMap<FName, TArray<TSharedPtr<FUICommandInfo>>> Commands;
 };
 
@@ -221,7 +224,6 @@ protected:
 	bool bShiftModifierOn = false;
 	int32 CtrlModifierId = 2;
 	bool bCtrlModifierOn = false;
-	// note we are using shift+ctrl for alt
 	int32 AltModiferId = 3;
 	bool bAltModifierOn = false;
 
diff --git a/Engine/Plugins/Editor/SequencerAnimTools/Source/SequencerAnimTools/Public/Trail.h b/Engine/Plugins/Editor/SequencerAnimTools/Source/SequencerAnimTools/Public/Trail.h
--- a/Engine/Plugins/Editor/SequencerAnimTools/Source/SequencerAnimTools/Public/Trail.h
+++ b/Engine/Plugins/Editor/SequencerAnimTools/Source/SequencerAnimTools/Public/Trail.h
@@ -5,6 +5,7 @@
 
 #include "Math/Box.h"
 #include "TrajectoryDrawInfo.h"
+#include "Tools/MotionTrailOptions.h"
 #include "HitProxies.h"
 #include "UObject/Object.h"
 #include "UObject/WeakObjectPtr.h"
@@ -16,6 +17,7 @@ struct FConvexVolume;
 class FEditorViewportClient;
 class FPrimitiveDrawInterface;
 class UMotionTrailEditorMode;
+class UMotionTrailToolOptions;
 
 namespace UE
 {
@@ -40,17 +42,18 @@ struct HBaseTrailProxy : public HHitProxy
 		,Guid(InGuid)
 	{}
 };
-struct HMotionTrailProxy : public HBaseTrailProxy
+
+struct HNewMotionTrailProxy : public HBaseTrailProxy
 {
 	DECLARE_HIT_PROXY();
 
 	FVector Point;
-	double Second;
+	FFrameNumber CurrentFrame;
 
-	HMotionTrailProxy(const FGuid& InGuid, const FVector& InPoint, double InSecond) 
+	HNewMotionTrailProxy(const FGuid& InGuid, const FVector& InPoint, const FFrameNumber& InFrame)
 		:HBaseTrailProxy(InGuid)
-		,Point(InPoint)
-		,Second(InSecond)
+		, Point(InPoint)
+		, CurrentFrame(InFrame)
 	{}
 
 };
@@ -63,6 +66,31 @@ enum class ETrailCacheState : uint8
 	NotUpdated = 3
 };
 
+//Status of trail, including possibly frames that must get updated on next evaluation
+struct FTrailCurrentStatus
+{
+	ETrailCacheState CacheState = ETrailCacheState::Stale;
+	TArray<FFrameNumber> FramesMustUpdate;
+};
+
+//struct to cacluate color. The same struct is meant to be used in same loop, and will be set up by the owner(trail hierarchy).
+struct  FColorState
+{
+	bool bFirstFrame = true;
+	FFrameNumber TicksPerFrame = FFrameNumber(100);
+	FFrameNumber StartFrame;
+	FFrameNumber SequencerTime;
+	bool bIsPinned = false;
+	EMotionTrailTrailStyle PinnedStyle;
+	FLinearColor CalculatedColor = FColor(0xffffff);
+	UMotionTrailToolOptions* Options = nullptr;
+
+	void Setup(FTrailHierarchy* TrailHierarchy);
+	void ReadyForTrail(bool bInIsPinned, EMotionTrailTrailStyle InPinnedStyle);
+	EMotionTrailTrailStyle GetStyle() const;
+};
+
+//main abstract class that trails implement
 class FTrail
 {
 public:
@@ -71,16 +99,14 @@ public:
 		: Owner(InOwner)
 		,CacheState(ETrailCacheState::Stale)
 		, bForceEvaluateNextTick(true)
-		, DrawInfo(nullptr)
 	{}
 
 	virtual ~FTrail() {}
 
-	// Public interface methods
-	struct FSceneContext
+	struct FNewSceneContext
 	{
+		bool bCheckForChange; 
 		FGuid YourNode;
-		const class FTrailEvaluateTimes& EvalTimes;
 		class FTrailHierarchy* TrailHierarchy;
 	};
 
@@ -88,54 +114,86 @@ public:
 	{ 
 		return (Owner.IsValid() ? Owner.Get() : nullptr);
 	}
+	//main function to update trail each tick
+	virtual FTrailCurrentStatus UpdateTrail(const FNewSceneContext& NewSceneContext) = 0;
+	//get the transform at a specified time
+	virtual void Interp(const FFrameNumber& Time, FTransform& OutTransform, FTransform& OutParentTransform) { OutTransform = FTransform::Identity; OutParentTransform = FTransform::Identity;
+	}
+	//when multiple frame update is finally finished
+	virtual void UpdateFinished(const TRange<FFrameNumber>& UpdatedRange, const TArray<int32>& IndicesToCalcluate) { ClearCachedData();  CacheState = ETrailCacheState::UpToDate; }
+	virtual FText GetName() const = 0;
 
-	virtual ETrailCacheState UpdateTrail(const FSceneContext& InSceneContext) = 0;
-	virtual FTrajectoryCache* GetTrajectoryTransforms() = 0;
-
+	//add set of important times for this trail to be used by the incremental evaluations, usually these are the edited times
+	virtual void AddImportantTimes(TSet<FFrameNumber>& InOutImportantTimes) {};
+	
 	//Additional Render/HitTest event handling for specific trails, usually let default renderer handle it
 	virtual void Render(const FGuid& Guid, const FSceneView* View, FPrimitiveDrawInterface* PDI) {};
+	virtual void RenderEvaluating(const FGuid& Guid, const FSceneView* View, FPrimitiveDrawInterface* PDI) {};
 	virtual void DrawHUD(const FSceneView* View, FCanvas* Canvas) {};
 	virtual bool HandleClick(const FGuid& Guid, FEditorViewportClient* InViewportClient, HHitProxy* HitProxy, FInputClick Click) { return false; }
 	virtual bool IsAnythingSelected() const { return false; }
 	virtual bool IsAnythingSelected(FVector& OutVectorPosition)const { return false;}
+	virtual bool IsAnythingSelected(TArray<FVector>& OutVectorPositions)const { return false; }
 	virtual bool IsTrailSelected() const { return false; }
 	virtual bool BoxSelect(FBox& InBox, bool InSelect = true) {return false;}
 	virtual bool FrustumSelect(const FConvexVolume& InFrustum, FEditorViewportClient* InViewportClient, bool InSelect = true)  { return false; }
-
+	virtual void SetOffsetMode() {};
 	virtual bool StartTracking() { return false; }
-	virtual bool ApplyDelta(const FVector& Pos, const FRotator& Rot, const FVector& WidgetLocation) { return false; }
+	virtual bool ApplyDelta(const FVector& Pos, const FRotator& Rot, const FVector& WidgetLocation, bool bApplyToOffset) { return false; }
 	virtual bool EndTracking() { return false; }
 	virtual void TranslateSelectedKeys(bool bRight) {};
 	virtual void DeleteSelectedKeys() {};
 	virtual void SelectNone() {};
-	virtual bool GetEditedTimes(const FTrailHierarchy* TrailHierarchy, const FFrameNumber& LastFrame, TArray<FFrameNumber>& OutEditedTimes) { return false; }
-	virtual void UpdateKeysInRange(const TRange<double>& ViewRange) {};
+	virtual void UpdateKeysInRange(const TRange<FFrameNumber>& ViewRange) {};
+	virtual void ClearCachedData() {};
+	virtual void HasStartedEvaluating() {};
+	virtual void ReadyToDrawTrail(FColorState& ColorState, const FCurrentFramesInfo* InCurrentFramesInfo,bool bIsEvaluating, bool bIsPinned) {};
+	virtual void ActorChanged(AActor* InActor) {};
 	// Optionally implemented methods
-	virtual TRange<double> GetEffectiveRange() const { return TRange<double>::Empty(); }
 	virtual TArray<FFrameNumber> GetKeyTimes() const { return TArray<FFrameNumber>(); }
 	virtual TArray<FFrameNumber> GetSelectedKeyTimes() const { return TArray<FFrameNumber>(); }
 
-	virtual void GetTrajectoryPointsForDisplay(const FDisplayContext& InDisplayContext, TArray<FVector>& OutPoints, TArray<double>& OutSeconds)
+	virtual void GetTrajectoryPointsForDisplay(const FCurrentFramesInfo& InCurrentFramesInfo, bool bIsEvaluating, TArray<FVector>& OutPoints, TArray<FFrameNumber>& OutFrames)
 	{
 		if (FTrajectoryDrawInfo* DI = GetDrawInfo())
 		{
-			DI->GetTrajectoryPointsForDisplay(InDisplayContext, OutPoints, OutSeconds);
+			DI->GetTrajectoryPointsForDisplay(GetOffsetTransform(), GetParentSpaceTransform(), InCurrentFramesInfo, bIsEvaluating, OutPoints, OutFrames);
 		}
 	}
-	virtual void GetTickPointsForDisplay(const FDisplayContext& InDisplayContext, TArray<FVector2D>& OutTicks, TArray<FVector2D>& OutTickTangents)
+	virtual void GetTickPointsForDisplay(const FTrailScreenSpaceTransform& InScreenSpaceTransform, const FCurrentFramesInfo& InCurrentFramesInfo, bool bIsEvaluating, TArray<FVector2D>& OutTicks, TArray<FVector2D>& OutTickTangents)
 	{
 		if (FTrajectoryDrawInfo* DI = GetDrawInfo())
 		{
-			DI->GetTickPointsForDisplay(InDisplayContext, OutTicks, OutTickTangents);
+			DI->GetTickPointsForDisplay(GetOffsetTransform(), GetParentSpaceTransform(), InScreenSpaceTransform, InCurrentFramesInfo, bIsEvaluating, OutTicks, OutTickTangents);
 		}
 	}
 
+	//offset the trail from current position
+	virtual FTransform GetOffsetTransform() const { return FTransform::Identity; }
+	virtual void ClearOffsetTransform() {};
+	virtual bool HasOffsetTransform() { return !GetOffsetTransform().Equals(FTransform::Identity); }
+	//set space and get space transform
+	virtual FTransform GetParentSpaceTransform() const { return FTransform::Identity; }
+	virtual void SetSpace(AActor* InActor, const FName& InComponentName) {};
+	virtual void ClearSpace() {};
+
 	virtual void ForceEvaluateNextTick() { bForceEvaluateNextTick = true; }
 
+	virtual void GetColor(const FFrameNumber& CurrentTime, FColorState& InOutColorState);
+	//get and restore trail from it's data
+	class FMotionTrailState
+	{
+	public:
+		virtual ~FMotionTrailState() {}
+		virtual void RestoreTrail(FTrailHierarchy* TrailHierarchy) = 0;
+	};
+	virtual TUniquePtr<FMotionTrailState> GetMotionTrailState() const { return nullptr; }
 
-	FTrajectoryDrawInfo* GetDrawInfo() { return DrawInfo.Get(); }
+	FTrajectoryDrawInfo* GetDrawInfo() const { return DrawInfo.Get(); }
 
 	ETrailCacheState GetCacheState() const { return CacheState; }
+
+
 protected:
 
 	TWeakObjectPtr<UObject> Owner;
@@ -143,6 +201,8 @@ protected:
 
 	bool bForceEvaluateNextTick;
 	TUniquePtr<FTrajectoryDrawInfo> DrawInfo;
+
+
 };
 } // namespace MovieScene
 } // namespace UE
