Motion Trails: Fix ticks getting drawn behind, need to not draw if projected W is negative

diff --git a/Engine/Plugins/Editor/SequencerAnimTools/Source/SequencerAnimTools/Private/TrajectoryDrawInfo.cpp b/Engine/Plugins/Editor/SequencerAnimTools/Source/SequencerAnimTools/Private/TrajectoryDrawInfo.cpp
--- a/Engine/Plugins/Editor/SequencerAnimTools/Source/SequencerAnimTools/Private/TrajectoryDrawInfo.cpp
+++ b/Engine/Plugins/Editor/SequencerAnimTools/Source/SequencerAnimTools/Private/TrajectoryDrawInfo.cpp
@@ -27,12 +27,15 @@ TOptional<FVector2D> FTrailScreenSpaceTransform::ProjectPoint(const FVector& Poi
 {
 
 	FVector2D PixelLocation;
-	if (View->WorldToPixel(Point, PixelLocation))
+	const FVector4 ScreenPoint = View->WorldToScreen(Point);
+	if (ScreenPoint.W > 0.0)
 	{
-		PixelLocation /= DPIScale;
-		return PixelLocation;
+		if (View->ScreenToPixel(ScreenPoint, PixelLocation))
+		{
+			PixelLocation /= DPIScale;
+			return PixelLocation;
+		}
 	}
-
 	return TOptional<FVector2D>();
 }
 
