Anim In Engine: When drawing trail in a different space need to update trail if that space or any of it's parents have changed.

diff --git a/Engine/Plugins/Editor/SequencerAnimTools/Source/SequencerAnimTools/Private/Sequencer/MovieSceneTransformTrail.cpp b/Engine/Plugins/Editor/SequencerAnimTools/Source/SequencerAnimTools/Private/Sequencer/MovieSceneTransformTrail.cpp
--- a/Engine/Plugins/Editor/SequencerAnimTools/Source/SequencerAnimTools/Private/Sequencer/MovieSceneTransformTrail.cpp
+++ b/Engine/Plugins/Editor/SequencerAnimTools/Source/SequencerAnimTools/Private/Sequencer/MovieSceneTransformTrail.cpp
@@ -33,119 +33,13 @@ namespace SequencerAnimTools
 using namespace UE::AIE;
 
 
-/****************************************************************************************************************
-*
-*  FCachedTracksAndSections
-*
-*******************************************************************************************************************/
-
-bool FCachedTracksAndSections::IsCacheEqual(const TArray <TWeakObjectPtr<UMovieSceneTrack>>& InTracks)
-{
-	if (Tracks.Num() != InTracks.Num())
-	{
-		CreateFrom(InTracks);
-		return false;
-	}
-	else
-	{
-		for (const TWeakObjectPtr<UMovieSceneTrack>& WeakTrack : InTracks)
-		{
-			if (WeakTrack.IsValid())
-			{
-				if (FCachedTrackAndSections* CachedTrackAndSections = Tracks.Find(WeakTrack))
-				{
-					if (WeakTrack->GetSignature() != CachedTrackAndSections->Track.CachedGuid)
-					{
-						CreateFrom(InTracks);
-						return false;
-					}
-					const TArray<UMovieSceneSection*>& Sections = WeakTrack->GetAllSections();
-					if (Sections.Num() != CachedTrackAndSections->Sections.Num())
-					{
-						CreateFrom(InTracks);
-						return false;
-					}
-					for (const UMovieSceneSection* Section : Sections)
-					{
-						if (FCachedGuidsPerSignedObject* SignedGuid = CachedTrackAndSections->Sections.Find(Section))
-						{
-							if (Section->GetSignature() != SignedGuid->CachedGuid)
-							{
-								CreateFrom(InTracks);
-								return false;
-							}
-						}
-						else
-						{
-							CreateFrom(InTracks);
-							return false;
-						}
-					}
-				}
-				else
-				{
-					CreateFrom(InTracks);
-					return false;
-				}
-			}
-		}
-		return true;
-	}
-}
-
-void  FCachedTracksAndSections::CreateFrom(const TArray <TWeakObjectPtr<UMovieSceneTrack>>& InTracks)
-{
-	Tracks.Reset();
-	for (const TWeakObjectPtr<UMovieSceneTrack>& WeakTrack : InTracks)
-	{
-		if (WeakTrack.IsValid())
-		{
-			FCachedTrackAndSections Track;
-			Track.Track.SignedObject = WeakTrack;
-			Track.Track.CachedGuid = WeakTrack->GetSignature();
-			const TArray<UMovieSceneSection*>& Sections = WeakTrack->GetAllSections();
-			for (const UMovieSceneSection* Section : Sections)
-			{
-				if (Section)
-				{
-					FCachedGuidsPerSignedObject SignedObjectGuid;
-					SignedObjectGuid.SignedObject = Section;
-					SignedObjectGuid.CachedGuid = Section->GetSignature();
-					Track.Sections.Add(Section,SignedObjectGuid);
-
-				}
-			}
-			Tracks.Add(WeakTrack,Track);
-		}
-	}
-}
-
-UMovieSceneSection* FCachedTracksAndSections::GetSection() const
-{
-	for (const TPair<TWeakObjectPtr<UMovieSceneTrack>, FCachedTrackAndSections>& Pair : Tracks)
-	{
-		if (Pair.Key.IsValid())
-		{
-			UMovieSceneSection* Section = Pair.Key.Get()->GetSectionToKey();
-			if (Section)
-			{
-				return Section;
-			}
-			else if (Pair.Key.Get()->GetAllSections().Num() > 0)
-			{
-				return Pair.Key.Get()->GetAllSections()[0];
-			}
-		}
-	}
-	return nullptr;
-}
 /****************************************************************************************************************
 *
 *   Optional Parent Space
 *
 *****************************************************************************************************************/
 
-void FOptionalParentSpace::SetSpace(AActor* InActor, const FName& InComponentName)
+void FOptionalParentSpace::SetSpace(TSharedPtr<ISequencer>&  InSequencer, AActor* InActor, const FName& InComponentName)
 {
 	if (InActor)
 	{
@@ -173,6 +67,22 @@ void FOptionalParentSpace::SetSpace(AActor* InActor, const FName& InComponentNam
 		{
 			ParentSpace.ParentTransforms = MakeShared<FArrayOfTransforms>();
 		}
+
+		check(InSequencer);
+		TArray<FMovieSceneObjectBindingID> Bindings;
+		FMovieSceneEvaluationState* State = InSequencer->GetEvaluationState();
+		State->FilterObjectBindings(ParentSpace.Actor.Component.Get(), *InSequencer, &Bindings);
+		if (Bindings.Num() <= 0)
+		{
+			AActor* ParentActor = ParentSpace.Actor.Component->GetOwner();
+			State->FilterObjectBindings(ParentActor, *InSequencer, &Bindings);
+		}
+		UMovieScene* MovieScene = InSequencer->GetFocusedMovieSceneSequence()->GetMovieScene();
+		TWeakObjectPtr<UMovieSceneTrack> MainTrack = nullptr;
+		if (Bindings.Num() > 0)
+		{
+			SpaceBindingID = Bindings[0].GetGuid();
+		}
 	}
 }
 
@@ -180,6 +90,7 @@ void FOptionalParentSpace::ClearSpace()
 {
 	bIsValid = false;
 	ParentSpace.Actor.Actor = nullptr;
+	SpaceBindingID = FGuid();
 }
 
 /****************************************************************************************************************
@@ -295,28 +206,52 @@ void FMovieSceneTransformTrail::ActorChanged(AActor* InActor)
 	}
 }
 
-bool FMovieSceneTransformTrail::TrailHasChanged()
+bool FMovieSceneTransformTrail::BindingHasChanged(const FGuid& InBindingID, const USceneComponent* InComponent, UE::AIE::FSequencerTransformDependencies& InOutDependencies) 
 {
 	if(TSharedPtr<ISequencer> SequencerPtr = WeakSequencer.Pin())
 	{ 
-		if (USceneComponent* Component = Cast<USceneComponent>(GetOwner()))
+		if (InComponent)
 		{
 			UMovieScene* MovieScene = SequencerPtr->GetFocusedMovieSceneSequence()->GetMovieScene();
-			if (FMovieSceneBinding* Binding = MovieScene->FindBinding(BindingID))
+			if (FMovieSceneBinding* Binding = MovieScene->FindBinding(InBindingID))
 			{
-				AActor* Actor = Component->GetOwner();
+				AActor* Actor = InComponent->GetOwner();
 				UE::AIE::FSequencerTransformDependencies NewTransformDependencies;
-				TArray<UMovieSceneTrack*> Tracks = FSequencerTransformDependencies::GetTransformAffectingTracks(SequencerPtr->GetFocusedMovieSceneSequence()->GetMovieScene(),
+				TArray<UMovieSceneTrack*> Tracks = FSequencerTransformDependencies::GetTransformAffectingTracks(MovieScene,
 					*Binding);
 
 				NewTransformDependencies.CalculateDependencies(WeakSequencer.Pin().Get(), Actor, Tracks);
-				if (TransformDependencies.Compare(NewTransformDependencies) == false)
+				if (InOutDependencies.Compare(NewTransformDependencies) == false)
 				{
-					TransformDependencies.CopyFrom(NewTransformDependencies);
+					InOutDependencies.CopyFrom(NewTransformDependencies);
 					return true;
 				}
 			}
-		
+		}
+	}
+	return false;
+}
+void FMovieSceneTransformTrail::SetSpace(AActor* InActor, const FName& InComponentName)
+{
+	TSharedPtr<ISequencer> SequencerPtr = GetSequencer();
+	ParentSpace.SetSpace(SequencerPtr, InActor, InComponentName);
+}
+
+bool FMovieSceneTransformTrail::TrailOrSpaceHasChanged() 
+{
+	if (USceneComponent* Component = Cast<USceneComponent>(GetOwner()))
+	{
+		if (BindingHasChanged(BindingID, Component, TransformDependencies))
+		{
+			return true;
+		}
+	}
+
+	if (ParentSpace.bIsValid == true)
+	{
+		if (BindingHasChanged(ParentSpace.SpaceBindingID, ParentSpace.ParentSpace.Actor.Component.Get(), TransformDependencies))
+		{
+			return true;
 		}
 	}
 	return false;
@@ -328,7 +263,7 @@ FTrailCurrentStatus FMovieSceneTransformTrail::UpdateTrail(const FNewSceneContex
 	
 	CachedHierarchyGuid = InSceneContext.YourNode;
 	
-	const bool bTrackUnchanged = (InSceneContext.bCheckForChange || TransformDependencies.IsEmpty() ) ? !TrailHasChanged() : true;
+	const bool bTrackUnchanged = (InSceneContext.bCheckForChange || TransformDependencies.IsEmpty() ) ? !TrailOrSpaceHasChanged() : true;
 
 	UMovieSceneSection* Section = GetSection();
 	TSharedPtr<ISequencer> Sequencer = WeakSequencer.Pin();
diff --git a/Engine/Plugins/Editor/SequencerAnimTools/Source/SequencerAnimTools/Private/Sequencer/MovieSceneTransformTrail.h b/Engine/Plugins/Editor/SequencerAnimTools/Source/SequencerAnimTools/Private/Sequencer/MovieSceneTransformTrail.h
--- a/Engine/Plugins/Editor/SequencerAnimTools/Source/SequencerAnimTools/Private/Sequencer/MovieSceneTransformTrail.h
+++ b/Engine/Plugins/Editor/SequencerAnimTools/Source/SequencerAnimTools/Private/Sequencer/MovieSceneTransformTrail.h
@@ -38,15 +38,6 @@ struct FCachedTrackAndSections
 	TMap<TWeakObjectPtr<const UMovieSceneSection>, FCachedGuidsPerSignedObject> Sections;
 };
 
-struct FCachedTracksAndSections
-{
-	TMap<TWeakObjectPtr<UMovieSceneTrack>, FCachedTrackAndSections> Tracks;
-	//checks if cache is equal and recreates it
-	bool IsCacheEqual(const TArray <TWeakObjectPtr<UMovieSceneTrack>>& InTracks);
-	UMovieSceneSection* GetSection() const;
-	void CreateFrom(const TArray <TWeakObjectPtr<UMovieSceneTrack>>& InTracks);
-};
-
 //space that a trail may be in, includes Actor/ optional component and it's sequencer tracks/sections
 struct FOptionalParentSpace
 {
@@ -54,11 +45,13 @@ struct FOptionalParentSpace
 
 	UE::AIE::FActorAndWorldTransforms ParentSpace; //this contains the actor/component
 	
-	FCachedTracksAndSections CachedTracksGuids;
 	bool bIsValid;
 
+	FGuid SpaceBindingID;
+	UE::AIE::FSequencerTransformDependencies SpaceTransformDependencies;
+
 	//set space and make it valid
-	void SetSpace(AActor* InActor, const FName& InComponentName);
+	void SetSpace(TSharedPtr<ISequencer>& InSequencerPtr, AActor* InActor, const FName& InComponentName);
 	void ClearSpace();
 };
 
@@ -85,7 +78,7 @@ public:
 	virtual void SetOffsetMode(); 
 
 	virtual FTransform GetParentSpaceTransform() const;
-	virtual void SetSpace(AActor* InActor, const FName& InComponentName) override { ParentSpace.SetSpace(InActor, InComponentName); }
+	virtual void SetSpace(AActor* InActor, const FName& InComponentName) override;
 	virtual void ClearSpace() override { ParentSpace.ClearSpace(); 	ClearCachedData();}
 	virtual void ClearCachedData() override;
 	virtual void ReadyToDrawTrail(FColorState& ColorState, const FCurrentFramesInfo* InCurrentFramesInfo, bool bIsEvaluating, bool bIsPinned) override;
@@ -149,7 +142,9 @@ protected:
 	//cached guid in the hierarchy
 	FGuid CachedHierarchyGuid;
 protected:
-	bool TrailHasChanged();
+
+	bool TrailOrSpaceHasChanged() ;
+	bool BindingHasChanged(const FGuid& InBinding, const USceneComponent* InComponent, UE::AIE::FSequencerTransformDependencies& InDependencies);
 
 protected: 
 	FGuid BindingID;
