ITF:  Add option to not exit tool on presave, this is now used by the motion trail tool

diff --git a/Engine/Plugins/Editor/SequencerAnimTools/Source/SequencerAnimTools/Private/SequencerToolsEditMode.cpp b/Engine/Plugins/Editor/SequencerAnimTools/Source/SequencerAnimTools/Private/SequencerToolsEditMode.cpp
--- a/Engine/Plugins/Editor/SequencerAnimTools/Source/SequencerAnimTools/Private/SequencerToolsEditMode.cpp
+++ b/Engine/Plugins/Editor/SequencerAnimTools/Source/SequencerAnimTools/Private/SequencerToolsEditMode.cpp
@@ -20,8 +20,11 @@
 FEditorModeID USequencerToolsEditMode::ModeName = TEXT("SequencerToolsEditMode");
 
 
-USequencerToolsEditMode::USequencerToolsEditMode() 
+USequencerToolsEditMode::USequencerToolsEditMode()
+	: bDeactivateOnPIEStartStateToRestore(false)
+	, bDeactivateOnSaveWorldToRestore(false)
 {
+
 	Info = FEditorModeInfo(
 		ModeName,
 		LOCTEXT("ModeName", "Sequencer Tools"),
@@ -39,6 +42,8 @@ void USequencerToolsEditMode::Enter()
 {
 	Super::Enter();
 
+	Owner->OnEditorModeIDChanged().AddUObject(this, &USequencerToolsEditMode::LocalOnModeActivated);
+
 	FLevelEditorModule* LevelEditorModule = FModuleManager::GetModulePtr<FLevelEditorModule>(TEXT("LevelEditor"));
 	
 	if (LevelEditorModule)
@@ -73,10 +78,32 @@ void USequencerToolsEditMode::Exit()
 			UE::TransformGizmoUtil::DeregisterTransformGizmoContextObject(GetInteractiveToolsContext(EToolsContextScope::EdMode));
 		}
 	}
+	Owner->OnEditorModeIDChanged().RemoveAll(this);
 
 	Super::Exit();
 }
 
+void USequencerToolsEditMode::LocalOnModeActivated(const FEditorModeID& InID, bool bIsActive)
+{
+	if (InID == GetID())
+	{
+		UEditorInteractiveToolsContext* EditorInteractiveToolsContext = GetInteractiveToolsContext(EToolsContextScope::Editor);
+		if (bIsActive)
+		{
+			bDeactivateOnPIEStartStateToRestore = EditorInteractiveToolsContext->GetDeactivateToolsOnPIEStart();
+			EditorInteractiveToolsContext->SetDeactivateToolsOnPIEStart(false);
+
+			bDeactivateOnSaveWorldToRestore = EditorInteractiveToolsContext->GetDeactivateToolsOnSaveWorld();
+			EditorInteractiveToolsContext->SetDeactivateToolsOnSaveWorld(false);
+		}
+		else
+		{
+			EditorInteractiveToolsContext->SetDeactivateToolsOnPIEStart(bDeactivateOnPIEStartStateToRestore);
+			EditorInteractiveToolsContext->SetDeactivateToolsOnSaveWorld(bDeactivateOnSaveWorldToRestore);
+		}
+	}
+}
+
 bool USequencerToolsEditMode::IsCompatibleWith(FEditorModeID OtherModeID) const
 {
 	// Compatible with all modes similar to FSequencerEdMode
diff --git a/Engine/Plugins/Editor/SequencerAnimTools/Source/SequencerAnimTools/Public/SequencerToolsEditMode.h b/Engine/Plugins/Editor/SequencerAnimTools/Source/SequencerAnimTools/Public/SequencerToolsEditMode.h
--- a/Engine/Plugins/Editor/SequencerAnimTools/Source/SequencerAnimTools/Public/SequencerToolsEditMode.h
+++ b/Engine/Plugins/Editor/SequencerAnimTools/Source/SequencerAnimTools/Public/SequencerToolsEditMode.h
@@ -28,5 +28,10 @@ public:
 	virtual bool UsesPropertyWidgets() const override { return true; }
 	virtual bool UsesTransformWidget() const { return false; }
 	virtual bool IsCompatibleWith(FEditorModeID OtherModeID) const override;
+private:
 
+	void LocalOnModeActivated(const FEditorModeID& InID, bool bIsActive); //UEdMode has this already defined but it's not virtual
+
+	bool bDeactivateOnPIEStartStateToRestore;
+	bool bDeactivateOnSaveWorldToRestore;
 };
diff --git a/Engine/Source/Editor/UnrealEd/Private/Tools/EdModeInteractiveToolsContext.cpp b/Engine/Source/Editor/UnrealEd/Private/Tools/EdModeInteractiveToolsContext.cpp
--- a/Engine/Source/Editor/UnrealEd/Private/Tools/EdModeInteractiveToolsContext.cpp
+++ b/Engine/Source/Editor/UnrealEd/Private/Tools/EdModeInteractiveToolsContext.cpp
@@ -408,7 +408,10 @@ void UEditorInteractiveToolsContext::TerminateActiveToolsOnPIEStart()
 }
 void UEditorInteractiveToolsContext::TerminateActiveToolsOnSaveWorld()
 {
-	DeactivateAllActiveTools(EToolShutdownType::Accept);
+	if (bDeactivateOnSaveWorld)
+	{
+		DeactivateAllActiveTools(EToolShutdownType::Accept);
+	}
 }
 void UEditorInteractiveToolsContext::TerminateActiveToolsOnWorldTearDown()
 {
@@ -861,6 +864,11 @@ void UEditorInteractiveToolsContext::SetDeactivateToolsOnPIEStart(bool bDeactiva
 	bDeactivateOnPIEStart = bDeactivateTools;
 }
 
+void UEditorInteractiveToolsContext::SetDeactivateToolsOnSaveWorld(bool bDeactivateTools)
+{
+	bDeactivateOnSaveWorld = bDeactivateTools;
+}
+
 void UModeManagerInteractiveToolsContext::Tick(FEditorViewportClient* ViewportClient, float DeltaTime)
 {
 	UEditorInteractiveToolsContext::Tick(ViewportClient, DeltaTime);
diff --git a/Engine/Source/Editor/UnrealEd/Public/Tools/EdModeInteractiveToolsContext.h b/Engine/Source/Editor/UnrealEd/Public/Tools/EdModeInteractiveToolsContext.h
--- a/Engine/Source/Editor/UnrealEd/Public/Tools/EdModeInteractiveToolsContext.h
+++ b/Engine/Source/Editor/UnrealEd/Public/Tools/EdModeInteractiveToolsContext.h
@@ -185,6 +185,17 @@ public:
 	/** @return true if tools should be deactivated when PIE is started */
 	bool GetDeactivateToolsOnPIEStart() const { return bDeactivateOnPIEStart; }
 
+	/*
+	 * Configure whether tools should shutdown when the world is saving.
+	 * By default, the context will shut down any active tools on save.
+	 * Setting this to true will prevent this, but the system makes no promises about whether tools work properly across a save.
+	 * Therefore, do not use this if you do not know for certain that it is safe to persist tools across a save in your situation.
+	 */
+	void SetDeactivateToolsOnSaveWorld(bool bDeactivateTools);
+
+	/** @return true if tools should be deactivated when save is started */
+	bool GetDeactivateToolsOnSaveWorld() const { return bDeactivateOnSaveWorld; }
+
 protected:
 
 	// we hide these 
@@ -243,6 +254,7 @@ private:
 
 	bool bDeactivateOnPIEStart = true;
 
+	bool bDeactivateOnSaveWorld = true;
 	bool bIsActive = false;
 };
 
