Anim In Engine: followup to not exiting motion trail tool when saving, need to handle spawnables getting recreated.

diff --git a/Engine/Plugins/Editor/SequencerAnimTools/Source/SequencerAnimTools/Private/Sequencer/MovieSceneTransformTrail.cpp b/Engine/Plugins/Editor/SequencerAnimTools/Source/SequencerAnimTools/Private/Sequencer/MovieSceneTransformTrail.cpp
--- a/Engine/Plugins/Editor/SequencerAnimTools/Source/SequencerAnimTools/Private/Sequencer/MovieSceneTransformTrail.cpp
+++ b/Engine/Plugins/Editor/SequencerAnimTools/Source/SequencerAnimTools/Private/Sequencer/MovieSceneTransformTrail.cpp
@@ -14,6 +14,7 @@
 #include "ControlRigObjectBinding.h"
 #include "Tools/MotionTrailOptions.h"
 
+#include "MovieSceneObjectBindingID.h"
 #include "MovieSceneToolHelpers.h"
 #include "ActorForWorldTransforms.h"
 #include "Evaluation/MovieScenePlayback.h"
@@ -44,18 +45,14 @@ void FOptionalParentSpace::SetSpace(TSharedPtr<ISequencer>&  InSequencer, AActor
 	if (InActor)
 	{
 		bIsValid = true;
+		ComponentName = NAME_None; 
 		ParentSpace.Actor.Actor = InActor;
 		if (InComponentName != NAME_None)
 		{
-			TInlineComponentArray<USceneComponent*> Components(InActor);
-
-			for (USceneComponent* Component : Components)
+			if (USceneComponent* Component = GetComponentFromName(InActor, InComponentName))
 			{
-				if (Component->GetFName() == InComponentName)
-				{
-					ParentSpace.Actor.Component = Component;
-					break;
-				}
+				ParentSpace.Actor.Component = Component;
+				ComponentName = InComponentName;
 			}
 		}
 
@@ -77,8 +74,6 @@ void FOptionalParentSpace::SetSpace(TSharedPtr<ISequencer>&  InSequencer, AActor
 			AActor* ParentActor = ParentSpace.Actor.Component->GetOwner();
 			State->FilterObjectBindings(ParentActor, *InSequencer, &Bindings);
 		}
-		UMovieScene* MovieScene = InSequencer->GetFocusedMovieSceneSequence()->GetMovieScene();
-		TWeakObjectPtr<UMovieSceneTrack> MainTrack = nullptr;
 		if (Bindings.Num() > 0)
 		{
 			SpaceBindingID = Bindings[0].GetGuid();
@@ -86,11 +81,25 @@ void FOptionalParentSpace::SetSpace(TSharedPtr<ISequencer>&  InSequencer, AActor
 	}
 }
 
+USceneComponent* FOptionalParentSpace::GetComponentFromName(const AActor* InActor, const FName& InComponentName)
+{
+	TInlineComponentArray<USceneComponent*> Components(InActor);
+	for (USceneComponent* Component : Components)
+	{
+		if (Component->GetFName() == InComponentName)
+		{
+			return Component;
+		}
+	}
+	return nullptr;
+}
+
 void FOptionalParentSpace::ClearSpace()
 {
 	bIsValid = false;
 	ParentSpace.Actor.Actor = nullptr;
 	SpaceBindingID = FGuid();
+	ComponentName = NAME_None;
 }
 
 /****************************************************************************************************************
@@ -100,7 +109,7 @@ void FOptionalParentSpace::ClearSpace()
 *****************************************************************************************************************/
 
 
-FMovieSceneTransformTrail::FMovieSceneTransformTrail(const FGuid& InBindingID,USceneComponent* InOwner,  const bool bInIsVisible, TWeakObjectPtr<UMovieSceneTrack> InWeakTrack, TSharedPtr<ISequencer> InSequencer)
+FMovieSceneTransformTrail::FMovieSceneTransformTrail(const FGuid& InBindingID, USceneComponent* InOwner, const bool bInIsVisible, TWeakObjectPtr<UMovieSceneTrack> InWeakTrack, TSharedPtr<ISequencer> InSequencer)
 	: FTrail(InOwner)
 	, WeakSequencer(InSequencer)
 	, bIsSelected(false)
@@ -144,6 +153,89 @@ void FMovieSceneTransformTrail::UpdateFinished(const TRange<FFrameNumber>& Range
 		KeyTool->UpdateKeys();
 	}
 }
+TArray<UObject*> FMovieSceneTransformTrail::GetBoundObjects(TSharedPtr<ISequencer>& InSequencerPtr, const FGuid& InGuid)
+{
+	TArray<UObject*> BoundObjects;
+	if (InSequencerPtr.IsValid())
+	{
+		FMovieSceneSequenceID SequenceId = InSequencerPtr->GetEvaluationState()->FindSequenceId(InSequencerPtr->GetFocusedMovieSceneSequence());
+		FMovieSceneObjectBindingID ObjectBinding = UE::MovieScene::FFixedObjectBindingID(InGuid, SequenceId);
+
+		for (TWeakObjectPtr<> WeakObject : ObjectBinding.ResolveBoundObjects(InSequencerPtr->GetFocusedTemplateID(), *InSequencerPtr))
+		{
+			if (WeakObject.IsValid())
+			{
+				BoundObjects.Add(WeakObject.Get());
+			}
+		}
+	}
+	return BoundObjects;
+}
+//we need a sequencer for this so FTrail does not have an implementation
+void FMovieSceneTransformTrail::CheckAndUpdateObjects()
+{
+	if (TSharedPtr<ISequencer> SequencerPtr = GetSequencer())
+	{
+		TArray<UObject*> BoundObjects;
+		if (Owner.IsValid() == false)
+		{
+			BoundObjects = GetBoundObjects(SequencerPtr, BindingID);
+			if (BoundObjects.Num() > 0)
+			{
+				if (AActor* Actor = Cast<AActor>(BoundObjects[0]))
+				{
+					Owner = Actor->GetRootComponent();
+				}
+				else if (USceneComponent* Comp = Cast<USceneComponent>(BoundObjects[0]))
+				{
+					Owner = Comp;
+				}
+			}
+			BoundObjects.SetNum(0);
+		}
+		if (ParentSpace.bIsValid && ( (ParentSpace.ParentSpace.Actor.Actor.IsValid() == false)
+			|| (ParentSpace.ComponentName.IsNone() == false && ParentSpace.ParentSpace.Actor.Component.IsValid() == false) ))
+		{
+			BoundObjects = GetBoundObjects(SequencerPtr, ParentSpace.SpaceBindingID);
+			if (BoundObjects.Num() > 0)
+			{
+				if (AActor* Actor = Cast<AActor>(BoundObjects[0]))
+				{
+					ParentSpace.ParentSpace.Actor.Actor = Actor;
+					if (ParentSpace.ComponentName.IsNone() == false)
+					{
+						USceneComponent* Component = FOptionalParentSpace::GetComponentFromName(Actor, ParentSpace.ComponentName);
+						ParentSpace.ParentSpace.Actor.Component = Component;
+					}
+				}
+				else if (USceneComponent* Comp = Cast<USceneComponent>(BoundObjects[0]))
+				{
+					ParentSpace.ParentSpace.Actor.Component = Comp;
+					ParentSpace.ParentSpace.Actor.Actor = Comp->GetOwner();
+				}
+			}
+		}
+	}
+}
+
+bool FMovieSceneTransformTrail::HandleObjectsChanged(const TMap<UObject*, UObject*>& ReplacementMap) 
+{
+	bool bSomethingChanged = FTrail::HandleObjectsChanged(ReplacementMap);
+	if (ParentSpace.bIsValid)
+	{
+		if (UObject* const* NewObject = ReplacementMap.Find(ParentSpace.ParentSpace.Actor.Component.Get()))
+		{
+			ParentSpace.ParentSpace.Actor.Component = Cast<USceneComponent>(*NewObject);
+			bSomethingChanged = bSomethingChanged || true;
+		}
+		if (UObject* const* NewObject = ReplacementMap.Find(ParentSpace.ParentSpace.Actor.Actor.Get()))
+		{
+			ParentSpace.ParentSpace.Actor.Actor = Cast<AActor>(*NewObject);
+			bSomethingChanged = bSomethingChanged || true;
+		}
+	}
+	return bSomethingChanged;
+}
 
 FTransform FMovieSceneTransformTrail::GetParentSpaceTransform()  const
 {
@@ -152,7 +244,21 @@ FTransform FMovieSceneTransformTrail::GetParentSpaceTransform()  const
 		if (TSharedPtr<ISequencer> Sequencer = WeakSequencer.Pin())
 		{
 			FFrameTime Time = Sequencer->GetLocalTime().Time;
-			USceneComponent* ParentComponent = ParentSpace.ParentSpace.Actor.Component.IsValid() ? ParentSpace.ParentSpace.Actor.Component.Get() : ParentSpace.ParentSpace.Actor.Actor->GetRootComponent();
+			auto GetParentComponent = [this]() -> USceneComponent*
+			{
+				if (ParentSpace.ParentSpace.Actor.Component.IsValid())
+				{
+					return  ParentSpace.ParentSpace.Actor.Component.Get();
+				}
+				else if (ParentSpace.ParentSpace.Actor.Actor.IsValid())
+				{
+					return ParentSpace.ParentSpace.Actor.Actor->GetRootComponent();
+				}
+				return nullptr;
+			};
+
+			USceneComponent* ParentComponent = GetParentComponent();
+			
 			FTransform ParentSpaceTM = ParentComponent ? ParentComponent->GetComponentToWorld() : FTransform::Identity;				
 			if (PreviousParentSpaceTM.Equals(ParentSpaceTM) == false)
 			{
@@ -259,6 +365,9 @@ bool FMovieSceneTransformTrail::TrailOrSpaceHasChanged()
 
 FTrailCurrentStatus FMovieSceneTransformTrail::UpdateTrail(const FNewSceneContext& InSceneContext)
 {
+
+	CheckAndUpdateObjects();
+
 	FTrailCurrentStatus Status;
 	
 	CachedHierarchyGuid = InSceneContext.YourNode;
@@ -769,8 +878,20 @@ void FMovieSceneComponentTransformTrail::UpdateNeedsEvaluation(const FGuid& InTr
 		Actors.Actor = Component->GetOwner();
 		Actors.Component = Component;
 		SequencerHierarchy->EvaluateActor(InTrailGuid, Actors, ArrayOfTransforms, ParentArrayOfTransforms);
+	}
+}
+
+bool FMovieSceneComponentTransformTrail::HandleObjectsChanged(const TMap<UObject*, UObject*>& ReplacementMap)
+{
+	bool bSomethingChanged = FMovieSceneTransformTrail::HandleObjectsChanged(ReplacementMap);
 
+	if (UObject* const* NewObject = ReplacementMap.Find(Component.Get()))
+	{
+		Component = Cast<USceneComponent>(*NewObject);
+		bSomethingChanged = bSomethingChanged || true;
 	}
+		
+	return bSomethingChanged;
 }
 
 FText FMovieSceneComponentTransformTrail::GetName() const
@@ -987,6 +1108,19 @@ void FMovieSceneSocketTransformTrail::UpdateNeedsEvaluation(const FGuid& InTrail
 	}
 }
 
+bool FMovieSceneSocketTransformTrail::HandleObjectsChanged(const TMap<UObject*, UObject*>& ReplacementMap)
+{
+	bool bSomethingChanged = FMovieSceneTransformTrail::HandleObjectsChanged(ReplacementMap);
+
+	if (UObject* const* NewObject = ReplacementMap.Find(Component.Get()))
+	{
+		Component = Cast<USceneComponent>(*NewObject);
+		bSomethingChanged = bSomethingChanged || true;
+	}
+
+	return bSomethingChanged;
+}
+
 
 FText FMovieSceneSocketTransformTrail::GetName() const
 {
diff --git a/Engine/Plugins/Editor/SequencerAnimTools/Source/SequencerAnimTools/Private/Sequencer/MovieSceneTransformTrail.h b/Engine/Plugins/Editor/SequencerAnimTools/Source/SequencerAnimTools/Private/Sequencer/MovieSceneTransformTrail.h
--- a/Engine/Plugins/Editor/SequencerAnimTools/Source/SequencerAnimTools/Private/Sequencer/MovieSceneTransformTrail.h
+++ b/Engine/Plugins/Editor/SequencerAnimTools/Source/SequencerAnimTools/Private/Sequencer/MovieSceneTransformTrail.h
@@ -46,13 +46,15 @@ struct FOptionalParentSpace
 	UE::AIE::FActorAndWorldTransforms ParentSpace; //this contains the actor/component
 	
 	bool bIsValid;
-
+	FName ComponentName;
 	FGuid SpaceBindingID;
 	UE::AIE::FSequencerTransformDependencies SpaceTransformDependencies;
 
 	//set space and make it valid
 	void SetSpace(TSharedPtr<ISequencer>& InSequencerPtr, AActor* InActor, const FName& InComponentName);
 	void ClearSpace();
+	
+	static USceneComponent* GetComponentFromName(const AActor* InActor, const FName& InComponentName);
 };
 
 struct FDrawCacheData
@@ -106,8 +108,10 @@ public:
 	virtual TArray<FFrameNumber> GetSelectedKeyTimes() const override;
 	virtual void ForceEvaluateNextTick() override;
 	virtual void GetColor(const FFrameNumber& CurrentTime, FColorState& InOutColorState) override;
-	// End FTrail interface
+	virtual bool HandleObjectsChanged(const TMap<UObject*, UObject*>& ReplacementMap) override;
 
+	// End FTrail interface
+	
 	TSharedPtr<ISequencer> GetSequencer() const { return WeakSequencer.Pin(); }
 	FGuid GetCachedHierarchyGuid() const { return CachedHierarchyGuid; }
 	UMovieSceneSection* GetSection() const;
@@ -127,6 +131,11 @@ protected:
 	//if selected or selected and offset mode we draw different colors
 	TOptional<FLinearColor> GetOverrideColor() const;
 
+protected:
+	//do to spawnables, etc.. objects that we have may have been deleted or recreated, need to update them if they have
+	virtual void CheckAndUpdateObjects();
+	static TArray<UObject*> GetBoundObjects(TSharedPtr<ISequencer>& InSequencerPtr, const FGuid& InGuid);
+
 protected:
 	//some interactive states/modes
 	bool bIsSelected;
@@ -198,6 +207,7 @@ public:
 	virtual void UpdateNeedsEvaluation(const FGuid& InTrailGuid, FSequencerTrailHierarchy* SequencerHierarchy) override;
 	virtual FText GetName() const override;
 	virtual TUniquePtr<FMotionTrailState> GetMotionTrailState() const;
+	virtual bool HandleObjectsChanged(const TMap<UObject*, UObject*>& ReplacementMap);
 
 
 private:
@@ -239,6 +249,7 @@ public:
 	virtual void UpdateNeedsEvaluation(const FGuid& InTrailGuid, FSequencerTrailHierarchy* SequencerHierarchy) override;
 	virtual FText GetName() const override;
 	virtual TUniquePtr<FMotionTrailState> GetMotionTrailState() const;
+	virtual bool HandleObjectsChanged(const TMap<UObject*, UObject*>& ReplacementMap);
 
 
 private:
diff --git a/Engine/Plugins/Editor/SequencerAnimTools/Source/SequencerAnimTools/Private/Sequencer/SequencerTrailHierarchy.cpp b/Engine/Plugins/Editor/SequencerAnimTools/Source/SequencerAnimTools/Private/Sequencer/SequencerTrailHierarchy.cpp
--- a/Engine/Plugins/Editor/SequencerAnimTools/Source/SequencerAnimTools/Private/Sequencer/SequencerTrailHierarchy.cpp
+++ b/Engine/Plugins/Editor/SequencerAnimTools/Source/SequencerAnimTools/Private/Sequencer/SequencerTrailHierarchy.cpp
@@ -170,6 +170,16 @@ void FSequencerTrailHierarchy::Initialize()
 		}
 	});
 
+	OnObjectsReplacedHandle = FCoreUObjectDelegates::OnObjectsReplaced.AddLambda([this](const TMap<UObject*, UObject*>& ReplacementMap)
+	{
+		for (const TPair<FGuid, TUniquePtr<FTrail>>& GuidTrailPair : GetAllTrails())
+		{
+			if (GuidTrailPair.Value.IsValid())
+			{
+				GuidTrailPair.Value->HandleObjectsChanged(ReplacementMap);
+			}
+		}
+	});
 
 	GEngine->OnLevelActorAdded().AddRaw(this, &FSequencerTrailHierarchy::OnActorChangedSomehow);
 	GEngine->OnLevelActorDeleted().AddRaw(this, &FSequencerTrailHierarchy::OnActorChangedSomehow);
@@ -190,6 +200,8 @@ void FSequencerTrailHierarchy::Initialize()
 		}
 		PreviouslyPinnedTrails.Reset();
 	}
+
+	
 	
 }
 
@@ -243,7 +255,7 @@ void FSequencerTrailHierarchy::Destroy()
 	//GEditor->UnregisterForUndo(this);
 	VisibilityManager.Reset();
 	UnRegisterMotionTrailOptionDelegates();
-
+	FCoreUObjectDelegates::OnObjectsReplaced.Remove(OnObjectsReplacedHandle);
 
 	UMotionTrailToolOptions::GetTrailOptions()->ResetPinnedItems();
 }
diff --git a/Engine/Plugins/Editor/SequencerAnimTools/Source/SequencerAnimTools/Private/Sequencer/SequencerTrailHierarchy.h b/Engine/Plugins/Editor/SequencerAnimTools/Source/SequencerAnimTools/Private/Sequencer/SequencerTrailHierarchy.h
--- a/Engine/Plugins/Editor/SequencerAnimTools/Source/SequencerAnimTools/Private/Sequencer/SequencerTrailHierarchy.h
+++ b/Engine/Plugins/Editor/SequencerAnimTools/Source/SequencerAnimTools/Private/Sequencer/SequencerTrailHierarchy.h
@@ -106,6 +106,7 @@ private:
 	FDelegateHandle OnActorAddedToSequencerHandle;
 	FDelegateHandle OnSelectionChangedHandle;
 	FDelegateHandle OnViewOptionsChangedHandle;
+	FDelegateHandle OnObjectsReplacedHandle;
 
 	struct FControlRigDelegateHandles
 	{
diff --git a/Engine/Plugins/Editor/SequencerAnimTools/Source/SequencerAnimTools/Private/Trail.cpp b/Engine/Plugins/Editor/SequencerAnimTools/Source/SequencerAnimTools/Private/Trail.cpp
--- a/Engine/Plugins/Editor/SequencerAnimTools/Source/SequencerAnimTools/Private/Trail.cpp
+++ b/Engine/Plugins/Editor/SequencerAnimTools/Source/SequencerAnimTools/Private/Trail.cpp
@@ -3,4 +3,18 @@
 #include "Trail.h"
 
 //MZ TODO LOOK AT REIMPLEMENTING CONSTANT TRAILS
-
+namespace UE
+{
+namespace SequencerAnimTools
+{
+bool FTrail::HandleObjectsChanged(const TMap<UObject*, UObject*>& ReplacementMap)
+{
+	if (UObject*const* NewObject = ReplacementMap.Find(Owner.Get()))
+	{
+		Owner = *NewObject;
+		return true;
+	}
+	return false;
+}
+}
+}
\ No newline at end of file
diff --git a/Engine/Plugins/Editor/SequencerAnimTools/Source/SequencerAnimTools/Public/Trail.h b/Engine/Plugins/Editor/SequencerAnimTools/Source/SequencerAnimTools/Public/Trail.h
--- a/Engine/Plugins/Editor/SequencerAnimTools/Source/SequencerAnimTools/Public/Trail.h
+++ b/Engine/Plugins/Editor/SequencerAnimTools/Source/SequencerAnimTools/Public/Trail.h
@@ -149,6 +149,8 @@ public:
 	virtual void HasStartedEvaluating() {};
 	virtual void ReadyToDrawTrail(FColorState& ColorState, const FCurrentFramesInfo* InCurrentFramesInfo,bool bIsEvaluating, bool bIsPinned) {};
 	virtual void ActorChanged(AActor* InActor) {};
+	virtual bool HandleObjectsChanged(const TMap<UObject*, UObject*>& ReplacementMap); //return true if there was a changed object
+
 	// Optionally implemented methods
 	virtual TArray<FFrameNumber> GetKeyTimes() const { return TArray<FFrameNumber>(); }
 	virtual TArray<FFrameNumber> GetSelectedKeyTimes() const { return TArray<FFrameNumber>(); }
