Anim In Engine: Motion Trail Update, changed the number of tick iterations per cache cycle to 60 frames, which is adjustable via the cvar Sequencer.MotionTrailTickIterations 60 and fix issue with PlayerMannequins flashing to last frame when calculating trails.

diff --git a/Engine/Plugins/Editor/SequencerAnimTools/Source/SequencerAnimTools/Private/Sequencer/SequencerTrailHierarchy.cpp b/Engine/Plugins/Editor/SequencerAnimTools/Source/SequencerAnimTools/Private/Sequencer/SequencerTrailHierarchy.cpp
--- a/Engine/Plugins/Editor/SequencerAnimTools/Source/SequencerAnimTools/Private/Sequencer/SequencerTrailHierarchy.cpp
+++ b/Engine/Plugins/Editor/SequencerAnimTools/Source/SequencerAnimTools/Private/Sequencer/SequencerTrailHierarchy.cpp
@@ -382,7 +382,6 @@ void FSequencerTrailHierarchy::EvaluateSequencerAndSetTransforms()
 			MustHaveTimes.Add(TickEvalRange.GetUpperBoundValue().Value);
 			FFrameNumber CurrentFrame = SequencerPtr->GetLocalTime().Time.GetFrame();
 
-			MustHaveTimes.Add(CurrentFrame);
 			if (TickViewRange.GetLowerBoundValue().Value != TickEvalRange.GetLowerBoundValue().Value)
 			{
 				MustHaveTimes.Add(TickViewRange.GetLowerBoundValue().Value);
@@ -406,7 +405,7 @@ void FSequencerTrailHierarchy::EvaluateSequencerAndSetTransforms()
 					}
 				}
 			}
-			CurrentFramesInfo.AddMustHaveTimes(MustHaveTimes);
+			CurrentFramesInfo.AddMustHaveTimes(MustHaveTimes, CurrentFrame);
 
 			bool bKeepCalculating = false;
 
diff --git a/Engine/Plugins/Editor/SequencerAnimTools/Source/SequencerAnimTools/Private/TrajectoryDrawInfo.cpp b/Engine/Plugins/Editor/SequencerAnimTools/Source/SequencerAnimTools/Private/TrajectoryDrawInfo.cpp
--- a/Engine/Plugins/Editor/SequencerAnimTools/Source/SequencerAnimTools/Private/TrajectoryDrawInfo.cpp
+++ b/Engine/Plugins/Editor/SequencerAnimTools/Source/SequencerAnimTools/Private/TrajectoryDrawInfo.cpp
@@ -6,7 +6,7 @@
 #include "Framework/Application/SlateApplication.h"
 
 
-int32 UE::SequencerAnimTools::FFrameCalculator::IndicesToCalculatePerBucket = 10;
+int32 UE::SequencerAnimTools::FFrameCalculator::IndicesToCalculatePerBucket = 60;
 static FAutoConsoleVariableRef CVarIndicesToCalculatePerBucket(
 	TEXT("Sequencer.MotionTrailTickIterations"),
 	UE::SequencerAnimTools::FFrameCalculator::IndicesToCalculatePerBucket,
@@ -59,6 +59,7 @@ void FFrameCalculator::SetUpFrameCalculator(const UE::AIE::FFrameTimeByIndex& In
 	}
 
 	TRange<FFrameNumber> Range = ViewRange;
+	Ranges.SetNum(0);
 	Ranges.Add(Range);
 	CurrentRangeStartIndex = 0;
 	CurrentRangeEndIndex = InCurrentFrameTimes.NumFrames - 1;
@@ -114,6 +115,7 @@ bool FFrameCalculator::CalculateIndices()
 				IndicesToCalculate.Add(CalculatedIndex);
 			}
 		}
+		IndicesToCalculate.Add(MustHaveIndices[MustHaveIndices.Num() -1]);
 	}
 	else
 	{
@@ -153,14 +155,20 @@ void FCurrentFramesInfo::SetUpFrameTimes(const TRange<FFrameNumber>& InEvalFrame
 	FrameCalculator.SetUpFrameCalculator(CurrentFrameTimes, ViewRange);
 }
 
-void FCurrentFramesInfo::AddMustHaveTimes(const TSet<FFrameNumber>& InMustHaveTimes)
+void FCurrentFramesInfo::AddMustHaveTimes(const TSet<FFrameNumber>& InMustHaveTimes, const FFrameNumber& CurrentFrame)
 {
 	TArray<int32> MustHaveIndices;
+	int32 CurrentIndex = CurrentFrameTimes.CalculateIndex(CurrentFrame);
 	for (const FFrameNumber& FrameNumber : InMustHaveTimes)
 	{
 		int32 Index = CurrentFrameTimes.CalculateIndex(FrameNumber);
-		MustHaveIndices.Add(Index);
+		if (Index != CurrentIndex)
+		{
+			MustHaveIndices.Add(Index);
+		}
 	}
+	//current frame last so it doesn't flash on certain rigs
+	MustHaveIndices.Add(CurrentIndex);
 	FrameCalculator.AddMustHaveIndices(MustHaveIndices);
 }
 
diff --git a/Engine/Plugins/Editor/SequencerAnimTools/Source/SequencerAnimTools/Public/TrajectoryDrawInfo.h b/Engine/Plugins/Editor/SequencerAnimTools/Source/SequencerAnimTools/Public/TrajectoryDrawInfo.h
--- a/Engine/Plugins/Editor/SequencerAnimTools/Source/SequencerAnimTools/Public/TrajectoryDrawInfo.h
+++ b/Engine/Plugins/Editor/SequencerAnimTools/Source/SequencerAnimTools/Public/TrajectoryDrawInfo.h
@@ -76,7 +76,7 @@ struct FCurrentFramesInfo
 {
 	void SetViewRange(const TRange<FFrameNumber>& InViewRange);
 	void SetUpFrameTimes(const TRange<FFrameNumber>& InEvalFrameRange, const FFrameNumber& InFrameStep);
-	void AddMustHaveTimes(const TSet<FFrameNumber>& InMustHaveTimes);
+	void AddMustHaveTimes(const TSet<FFrameNumber>& InMustHaveTimes, const FFrameNumber& InFrameNumber);
 	const TArray<int32>& IndicesToCalculate() const;
 	bool KeepCalculating();
 	void Reset();
