Anim In Engine: Tweak selected color to make it more visible
#trivial

diff --git a/Engine/Plugins/Animation/ControlRig/Source/ControlRigEditor/Private/Editor/ControlRigViewportToolbarExtensions.cpp b/Engine/Plugins/Animation/ControlRig/Source/ControlRigEditor/Private/Editor/ControlRigViewportToolbarExtensions.cpp
--- a/Engine/Plugins/Animation/ControlRig/Source/ControlRigEditor/Private/Editor/ControlRigViewportToolbarExtensions.cpp
+++ b/Engine/Plugins/Animation/ControlRig/Source/ControlRigEditor/Private/Editor/ControlRigViewportToolbarExtensions.cpp
@@ -715,6 +715,10 @@ void CreateAdvancedSubMenu(UToolMenu* InSubMenu, UMotionTrailToolOptions* Settin
 		const FName ColorProperty("KeyColor");
 		ColorSettings.AddEntry(CreateTrailColor(Settings, ColorProperty));
 	}
+	{
+		const FName ColorProperty("SelectedKeyColor");
+		ColorSettings.AddEntry(CreateTrailColor(Settings, ColorProperty));
+	}
 
 
 }
diff --git a/Engine/Plugins/Editor/SequencerAnimTools/Source/SequencerAnimTools/Private/Sequencer/MotionTrailMovieSceneKey.cpp b/Engine/Plugins/Editor/SequencerAnimTools/Source/SequencerAnimTools/Private/Sequencer/MotionTrailMovieSceneKey.cpp
--- a/Engine/Plugins/Editor/SequencerAnimTools/Source/SequencerAnimTools/Private/Sequencer/MotionTrailMovieSceneKey.cpp
+++ b/Engine/Plugins/Editor/SequencerAnimTools/Source/SequencerAnimTools/Private/Sequencer/MotionTrailMovieSceneKey.cpp
@@ -55,7 +55,7 @@ void FMotionTraiMovieScenelKeyTool::Render(const FGuid& Guid, const FSceneView*
 	{
 		if (CachedViewFrameRange.Contains(FrameKeyPair.Value->FrameNumber))
 		{
-			FLinearColor KeyColor = CachedSelection.Contains(FrameKeyPair.Value.Get()) ? FLinearColor(1.0,1.0,0.0f) :  UMotionTrailToolOptions::GetTrailOptions()->KeyColor;
+			FLinearColor KeyColor = CachedSelection.Contains(FrameKeyPair.Value.Get()) ? UMotionTrailToolOptions::GetTrailOptions()->SelectedKeyColor :  UMotionTrailToolOptions::GetTrailOptions()->KeyColor;
 			const FTransform Transform = FrameKeyPair.Value->Transform * OffsetTransform * ParentSpaceTransform;
 			if (bHitTesting)
 			{
@@ -175,9 +175,12 @@ bool FMotionTraiMovieScenelKeyTool::HandleClick(const FGuid& Guid, FEditorViewpo
 			}
 
 		}
-		else
+		else 
 		{
-			CachedSelection.Reset();
+			if (Click.bShiftIsDown == false && Click.bCtrlIsDown == false)
+			{
+				CachedSelection.Reset();
+			}
 			return false;
 		}
 		UpdateSelectedKeysTransform();
diff --git a/Engine/Source/Editor/Sequencer/Public/Tools/MotionTrailOptions.h b/Engine/Source/Editor/Sequencer/Public/Tools/MotionTrailOptions.h
--- a/Engine/Source/Editor/Sequencer/Public/Tools/MotionTrailOptions.h
+++ b/Engine/Source/Editor/Sequencer/Public/Tools/MotionTrailOptions.h
@@ -47,6 +47,7 @@ public:
 		, bShowKeys(true)
 		, bShowFrameNumber_DEPRECATED(false)
 		, KeyColor(1.0, 1.0, 1.0)
+		, SelectedKeyColor(0.8,0.8,0.0)
 		, KeySize(3.0f)
 		, bShowMarks(false)
 		, MarkColor(0.25,1.0,0.15)
@@ -118,6 +119,10 @@ public:
 	UPROPERTY(EditAnywhere, Category = Keys)
 	FLinearColor KeyColor;
 
+	/** The color of the selected keys */
+	UPROPERTY(EditAnywhere, Config, Category = AdvancedDisplay)
+	FLinearColor SelectedKeyColor;
+
 	/** The size of the keys */
 	UPROPERTY(EditAnywhere, Config, Category = Keys, Meta = (ClampMin = "0.0"))
 	double KeySize;
