Anim In Engine:  Motion Trail & Pivot Tool to use new TRS gizmo. From Benoit. Added hiding of default widget with pivot tool.

diff --git a/Engine/Plugins/Editor/SequencerAnimTools/Source/SequencerAnimTools/Private/BaseSequencerAnimTool.cpp b/Engine/Plugins/Editor/SequencerAnimTools/Source/SequencerAnimTools/Private/BaseSequencerAnimTool.cpp
--- a/Engine/Plugins/Editor/SequencerAnimTools/Source/SequencerAnimTools/Private/BaseSequencerAnimTool.cpp
+++ b/Engine/Plugins/Editor/SequencerAnimTools/Source/SequencerAnimTools/Private/BaseSequencerAnimTool.cpp
@@ -1,9 +1,38 @@
 // Copyright Epic Games, Inc. All Rights Reserved.
 
 #include "BaseSequencerAnimTool.h"
+#include "EditorInteractiveGizmoManager.h"
+#include "EditorGizmos/EditorTransformGizmoUtil.h"
+#include "EditorGizmos/TransformGizmo.h"
+#include "BaseGizmos/TransformProxy.h"
+#include "InteractiveGizmoManager.h"
+#include "BaseGizmos/CombinedTransformGizmo.h"
+
 
 #include UE_INLINE_GENERATED_CPP_BY_NAME(BaseSequencerAnimTool)
 
 UBaseSequencerAnimTool::UBaseSequencerAnimTool(const FObjectInitializer& ObjectInitializer)
 	: Super(ObjectInitializer)
 {}
+
+
+void FSequencerAnimToolHelpers::CreateGizmo(const FSequencerAnimToolHelpers::FGizmoData& InData, TObjectPtr <UCombinedTransformGizmo>& OutCombinedGizmo, TObjectPtr<UTransformGizmo>& OutTRSGizmo)
+{
+	const bool bUseNewTRS = UEditorInteractiveGizmoManager::UsesNewTRSGizmos();
+	OutTRSGizmo = bUseNewTRS ? UE::EditorTransformGizmoUtil::CreateTransformGizmo(InData.ToolManager, FString(), InData.Owner) : nullptr;
+	if (OutTRSGizmo)
+	{
+		OutTRSGizmo->TransformGizmoSource = nullptr;
+		OutTRSGizmo->SetActiveTarget(InData.TransformProxy, InData.ToolManager);
+		if (UGizmoElementHitMultiTarget* HitMultiTarget = Cast< UGizmoElementHitMultiTarget>(OutTRSGizmo->HitTarget))
+		{
+			HitMultiTarget->GizmoTransformProxy = InData.TransformProxy;
+		}
+	}
+	else
+	{
+		ETransformGizmoSubElements Elements = ETransformGizmoSubElements::StandardTranslateRotate;
+		OutCombinedGizmo = InData.GizmoManager->CreateCustomTransformGizmo(Elements, InData.Owner, InData.InstanceIdentifier);
+		OutCombinedGizmo->SetActiveTarget(InData.TransformProxy, InData.ToolManager);
+}
+}
diff --git a/Engine/Plugins/Editor/SequencerAnimTools/Source/SequencerAnimTools/Private/MotionTrailTool.cpp b/Engine/Plugins/Editor/SequencerAnimTools/Source/SequencerAnimTools/Private/MotionTrailTool.cpp
--- a/Engine/Plugins/Editor/SequencerAnimTools/Source/SequencerAnimTools/Private/MotionTrailTool.cpp
+++ b/Engine/Plugins/Editor/SequencerAnimTools/Source/SequencerAnimTools/Private/MotionTrailTool.cpp
@@ -15,9 +15,12 @@
 #include "InteractiveToolManager.h"
 // for raycast into World
 
+#include "EditorInteractiveGizmoManager.h"
 #include "ScopedTransaction.h"
 #include "EditorModeManager.h"
 #include "UnrealClient.h"
+#include "EditorGizmos/EditorTransformGizmoUtil.h"
+#include "EditorGizmos/TransformGizmo.h"
 
 #include UE_INLINE_GENERATED_CPP_BY_NAME(MotionTrailTool)
 
@@ -85,17 +88,27 @@ void UMotionTrailTool::Setup()
 	AddInputBehavior(RightClickBehavior);
 
 	TransformProxy = NewObject<UTransformProxy>(this);
-	//TransformGizmo = GetToolManager()->GetPairedGizmoManager()->Create3AxisTransformGizmo(this, UMotionTrailTool::TrailKeyTransformGizmoInstanceIdentifier);
-
-	FString GizmoIdentifier = TEXT("PivotToolGizmoIdentifier");
-	ETransformGizmoSubElements Elements = ETransformGizmoSubElements::StandardTranslateRotate;
-	TransformGizmo = GizmoManager->CreateCustomTransformGizmo(Elements, this, UMotionTrailTool::TrailKeyTransformGizmoInstanceIdentifier);
 
 	TransformProxy->OnTransformChanged.AddUObject(this, &UMotionTrailTool::GizmoTransformChanged);
 	TransformProxy->OnBeginTransformEdit.AddUObject(this, &UMotionTrailTool::GizmoTransformStarted);
 	TransformProxy->OnEndTransformEdit.AddUObject(this, &UMotionTrailTool::GizmoTransformEnded);
-	TransformGizmo->SetActiveTarget(TransformProxy, GetToolManager());
-	TransformGizmo->SetVisibility(false);
+
+	FSequencerAnimToolHelpers::FGizmoData GizmoData;
+	GizmoData.Owner = this;
+	GizmoData.TransformProxy = TransformProxy;
+	GizmoData.ToolManager = GetToolManager();
+	GizmoData.GizmoManager = GizmoManager;
+	GizmoData.InstanceIdentifier = UMotionTrailTool::TrailKeyTransformGizmoInstanceIdentifier;
+	FSequencerAnimToolHelpers::CreateGizmo(GizmoData, TransformGizmo, TRSGizmo);
+
+	if (TRSGizmo)
+	{
+		TRSGizmo->SetVisibility(false);
+	}
+	else if (TransformGizmo)
+	{
+		TransformGizmo->SetVisibility(false);
+	}
 	SetupIntegration();
 }
 
@@ -159,6 +172,10 @@ void UMotionTrailTool::SetupIntegration()
 
 void UMotionTrailTool::Shutdown(EToolShutdownType ShutdownType)
 {
+	if (UMotionTrailToolOptions::GetTrailOptions()->bShowTrails == true)
+	{
+		UMotionTrailToolOptions::GetTrailOptions()->bShowTrails = false;
+	}
 	Super::Shutdown(ShutdownType);
 	GizmoManager->DestroyAllGizmosByOwner(this);
 	ShutdownIntegration();
@@ -328,47 +345,67 @@ void UMotionTrailTool::GizmoTransformEnded(UTransformProxy* Proxy)
 
 void UMotionTrailTool::UpdateGizmoVisibility()
 {
-	if (TransformGizmo)
+	auto IsAnythingSelected = [this]()
 	{
-		bool bAnythingSelected = false;
-
 		for (const TUniquePtr<UE::SequencerAnimTools::FTrailHierarchy>& TrailHierarchy : TrailHierarchies)
 		{
-			bool bIsHandled = TrailHierarchy->IsAnythingSelected();
-			if (bIsHandled)
+			if (TrailHierarchy->IsAnythingSelected())
 			{
-				bAnythingSelected = true;
+				return true;
 			}
 		}
-		TransformGizmo->SetVisibility(bAnythingSelected);
+		return false;
+	};
+	
+	if (TransformGizmo)
+	{
+		TransformGizmo->SetVisibility(IsAnythingSelected());
+	}
+	
+	if (TRSGizmo)
+	{
+		TRSGizmo->SetVisibility(IsAnythingSelected());
 	}
 }
 
 void UMotionTrailTool::UpdateGizmoLocation()
 {
-	if (!TransformGizmo)
-	{
-		return;
-	}
 	UpdateGizmoVisibility();
-	FVector NewGizmoLocation(0., 0., 0.);
-	int NumSelected = 0;
-	for (const TUniquePtr<UE::SequencerAnimTools::FTrailHierarchy>& TrailHierarchy : TrailHierarchies)
+	
+	auto GetNewLocation = [this]()
 	{
-		FVector Location;
-		bool bIsHandled = TrailHierarchy->IsAnythingSelected(Location);
-		if (bIsHandled)
+		FVector NewGizmoLocation(0., 0., 0.);
+		int NumSelected = 0;
+		for (const TUniquePtr<UE::SequencerAnimTools::FTrailHierarchy>& TrailHierarchy : TrailHierarchies)
 		{
-			NewGizmoLocation += Location;
-			++NumSelected;
+			FVector Location;
+			bool bIsHandled = TrailHierarchy->IsAnythingSelected(Location);
+			if (bIsHandled)
+			{
+				NewGizmoLocation += Location;
+				++NumSelected;
+			}
+		}
+		if (NumSelected > 0)
+		{
+			NewGizmoLocation /= (double)NumSelected;
 		}
+		return NewGizmoLocation;
+	};
+	
+	GizmoTransform.SetLocation(GetNewLocation());
+
+	if (TransformGizmo)
+	{
+		TransformGizmo->ReinitializeGizmoTransform(GizmoTransform);
 	}
-	if (NumSelected > 0)
+
+	if (TRSGizmo)
 	{
-		NewGizmoLocation /= (double)NumSelected;
+		// note: this is a trick to avoid cycling notifications... (cf. UCombinedTransformGizmo::ReinitializeGizmoTransform)
+		TGuardValue<bool> PivotModeGuard(TransformProxy->bSetPivotMode, true);
+		TransformProxy->SetTransform(GizmoTransform);
 	}
-	GizmoTransform.SetLocation(NewGizmoLocation);
-	TransformGizmo->ReinitializeGizmoTransform(GizmoTransform);
 }
 
 
diff --git a/Engine/Plugins/Editor/SequencerAnimTools/Source/SequencerAnimTools/Private/SequencerAnimEditPivotTool.cpp b/Engine/Plugins/Editor/SequencerAnimTools/Source/SequencerAnimTools/Private/SequencerAnimEditPivotTool.cpp
--- a/Engine/Plugins/Editor/SequencerAnimTools/Source/SequencerAnimTools/Private/SequencerAnimEditPivotTool.cpp
+++ b/Engine/Plugins/Editor/SequencerAnimTools/Source/SequencerAnimTools/Private/SequencerAnimEditPivotTool.cpp
@@ -20,6 +20,7 @@
 #include "MovieSceneTimeUnit.h"
 #include "ControlRig.h"
 #include "ControlRigSequencerEditorLibrary.h"
+#include "EditorInteractiveGizmoManager.h"
 #include "LevelSequence.h"
 #include "LevelSequenceEditorBlueprintLibrary.h"
 #include "ILevelSequenceEditorToolkit.h"
@@ -36,6 +37,8 @@
 #include "Widgets/Layout/SUniformGridPanel.h"
 #include "Styling/StyleColors.h"
 #include "SLevelViewport.h"
+#include "EditorGizmos/EditorTransformGizmoUtil.h"
+#include "EditorGizmos/TransformGizmo.h"
 
 #define LOCTEXT_NAMESPACE "SequencerAnimTools"
 
@@ -268,13 +271,23 @@ void USequencerPivotTool::Setup()
 
 	TransformProxy = NewObject<UTransformProxy>(this);
 
-	FString GizmoIdentifier = TEXT("PivotToolGizmoIdentifier");
-	TransformGizmo = GizmoManager->Create3AxisTransformGizmo(this, GizmoIdentifier);
-
 	TransformProxy->OnTransformChanged.AddUObject(this, &USequencerPivotTool::GizmoTransformChanged);
 	TransformProxy->OnBeginTransformEdit.AddUObject(this, &USequencerPivotTool::GizmoTransformStarted);
 	TransformProxy->OnEndTransformEdit.AddUObject(this, &USequencerPivotTool::GizmoTransformEnded);
-	TransformGizmo->SetActiveTarget(TransformProxy, GetToolManager());
+
+	FSequencerAnimToolHelpers::FGizmoData GizmoData;
+	GizmoData.Owner = this;
+	GizmoData.TransformProxy = TransformProxy;
+	GizmoData.ToolManager = GetToolManager();
+	GizmoData.GizmoManager = GizmoManager;
+	GizmoData.InstanceIdentifier = TEXT("PivotToolGizmoIdentifier");
+	FSequencerAnimToolHelpers::CreateGizmo(GizmoData, TransformGizmo, TRSGizmo);
+	//turn default widget(transform gizmo) off when entering this mode if not off
+	if (GLevelEditorModeTools().GetShowWidget())
+	{
+		bSetShowWidget = true;
+		GLevelEditorModeTools().SetShowWidget(false); 
+	}
 
 	GetControlRigsAndSequencer(ControlRigs, SequencerPtr, &LevelSequence);
 
@@ -502,6 +515,10 @@ bool USequencerPivotTool::SetGizmoBasedOnSelection(bool bUseSaved)
 }
 void USequencerPivotTool::Shutdown(EToolShutdownType ShutdownType)
 {
+	if (bSetShowWidget.IsSet() && bSetShowWidget.GetValue() == true)
+	{
+		GLevelEditorModeTools().SetShowWidget(true); //turn widget back on
+	}
 	GizmoManager->DestroyAllGizmosByOwner(this);
 	RemoveDelegates();
 	RemoveAndDestroyPivotOverlay();
@@ -934,16 +951,26 @@ void USequencerPivotTool::UpdateGizmoVisibility()
 	{
 		TransformGizmo->SetVisibility(true);
 	}
+	
+	if (TRSGizmo)
+	{
+		TRSGizmo->SetVisibility(true);
+	}
 }
 
 void USequencerPivotTool::UpdateGizmoTransform()
 {
-	if (!TransformGizmo)
+	if (TransformGizmo)
 	{
-		return;
+		TransformGizmo->ReinitializeGizmoTransform(GizmoTransform);
 	}
 
-	TransformGizmo->ReinitializeGizmoTransform(GizmoTransform);
+	if (TRSGizmo)
+	{
+		// note: this is a trick to avoid cycling notifications... (cf. UCombinedTransformGizmo::ReinitializeGizmoTransform)
+		TGuardValue<bool> PivotModeGuard(TransformProxy->bSetPivotMode, true);
+		TransformProxy->SetTransform(GizmoTransform);
+	}
 
 }
 
diff --git a/Engine/Plugins/Editor/SequencerAnimTools/Source/SequencerAnimTools/Public/BaseSequencerAnimTool.h b/Engine/Plugins/Editor/SequencerAnimTools/Source/SequencerAnimTools/Public/BaseSequencerAnimTool.h
--- a/Engine/Plugins/Editor/SequencerAnimTools/Source/SequencerAnimTools/Public/BaseSequencerAnimTool.h
+++ b/Engine/Plugins/Editor/SequencerAnimTools/Source/SequencerAnimTools/Public/BaseSequencerAnimTool.h
@@ -6,6 +6,12 @@
 #include "UObject/Interface.h"
 #include "BaseSequencerAnimTool.generated.h"
 
+class UInteractiveToolManager;
+class UTransformProxy;
+class UCombinedTransformGizmo;
+class UTransformGizmo;
+class UInteractiveGizmoManager;
+
 UINTERFACE(MinimalAPI)
 class UBaseSequencerAnimTool : public UInterface
 {
@@ -18,3 +24,16 @@ class SEQUENCERANIMTOOLS_API IBaseSequencerAnimTool
 	virtual bool ProcessCommandBindings(const FKey Key, const bool bRepeat) const { return false; }
 
 };
+
+struct SEQUENCERANIMTOOLS_API FSequencerAnimToolHelpers
+{
+	struct FGizmoData
+	{
+		void* Owner = nullptr;
+		UInteractiveToolManager* ToolManager = nullptr;
+		UTransformProxy* TransformProxy = nullptr;
+		UInteractiveGizmoManager* GizmoManager= nullptr;
+		FString InstanceIdentifier;
+	};
+	static void CreateGizmo(const FGizmoData& InData, TObjectPtr <UCombinedTransformGizmo>& OutCombinedGizmo, TObjectPtr<UTransformGizmo>& OutTRSGizmo);
+};
diff --git a/Engine/Plugins/Editor/SequencerAnimTools/Source/SequencerAnimTools/Public/MotionTrailTool.h b/Engine/Plugins/Editor/SequencerAnimTools/Source/SequencerAnimTools/Public/MotionTrailTool.h
--- a/Engine/Plugins/Editor/SequencerAnimTools/Source/SequencerAnimTools/Public/MotionTrailTool.h
+++ b/Engine/Plugins/Editor/SequencerAnimTools/Source/SequencerAnimTools/Public/MotionTrailTool.h
@@ -25,7 +25,7 @@ class UClickDragInputBehavior;
 class UCombinedTransformGizmo;
 class FEditorModeTools;
 class USingleKeyCaptureBehavior;
-
+class UTransformGizmo;
 
 /** 
 * IClickBehaviorTarget Wrapper to handle left and right clicks
@@ -164,6 +164,8 @@ public:
 	TObjectPtr<UTransformProxy> TransformProxy;
 	UPROPERTY()
 	TObjectPtr<UCombinedTransformGizmo> TransformGizmo;
+	UPROPERTY()
+	TObjectPtr<UTransformGizmo> TRSGizmo = nullptr;
 
 	FTransform StartDragTransform;
 	bool bGizmoBeingDragged = false;
diff --git a/Engine/Plugins/Editor/SequencerAnimTools/Source/SequencerAnimTools/Public/SequencerAnimEditPivotTool.h b/Engine/Plugins/Editor/SequencerAnimTools/Source/SequencerAnimTools/Public/SequencerAnimEditPivotTool.h
--- a/Engine/Plugins/Editor/SequencerAnimTools/Source/SequencerAnimTools/Public/SequencerAnimEditPivotTool.h
+++ b/Engine/Plugins/Editor/SequencerAnimTools/Source/SequencerAnimTools/Public/SequencerAnimEditPivotTool.h
@@ -27,6 +27,8 @@ struct FRigControlElement;
 class ISequencer;
 class FUICommandList;
 class IAssetViewport;
+class UTransformGizmo;
+
 /*
 *  The way this sequencer pivot tool works is that
 *  it will use modify the incoming selections temp pivot
@@ -177,6 +179,9 @@ protected:
 
 	UPROPERTY()
 	TObjectPtr<UCombinedTransformGizmo> TransformGizmo = nullptr;
+	
+	UPROPERTY()
+	TObjectPtr<UTransformGizmo> TRSGizmo = nullptr;
 
 protected:
 	bool bInPivotMode = false; //pivot may be in 'free' mode or 'pivot' mode
@@ -232,6 +237,7 @@ protected:
 	
 	static FVector2f LastPivotOverlayLocation;
 	TSharedPtr<SWidget> PivotWidget;
+	TOptional<bool> bSetShowWidget;
 public:
 	void TryRemovePivotOverlay();
 	void TryShowPivotOverlay();
