Anim In Engine: Delete Unused TrajectoryCache.h file
#trivial

diff --git a/Engine/Plugins/Editor/SequencerAnimTools/Source/SequencerAnimTools/Private/Sequencer/MovieSceneTransformTrail.h b/Engine/Plugins/Editor/SequencerAnimTools/Source/SequencerAnimTools/Private/Sequencer/MovieSceneTransformTrail.h
--- a/Engine/Plugins/Editor/SequencerAnimTools/Source/SequencerAnimTools/Private/Sequencer/MovieSceneTransformTrail.h
+++ b/Engine/Plugins/Editor/SequencerAnimTools/Source/SequencerAnimTools/Private/Sequencer/MovieSceneTransformTrail.h
@@ -6,7 +6,6 @@
 #include "Trail.h"
 #include "MotionTrailMovieSceneKey.h"
 #include "MovieSceneTrack.h"
-#include "TrajectoryCache.h"
 #include "Tools/EvaluateSequencerTools.h"
 #include "Misc/Optional.h"
 
diff --git a/Engine/Plugins/Editor/SequencerAnimTools/Source/SequencerAnimTools/Private/TrailHierarchy.cpp b/Engine/Plugins/Editor/SequencerAnimTools/Source/SequencerAnimTools/Private/TrailHierarchy.cpp
--- a/Engine/Plugins/Editor/SequencerAnimTools/Source/SequencerAnimTools/Private/TrailHierarchy.cpp
+++ b/Engine/Plugins/Editor/SequencerAnimTools/Source/SequencerAnimTools/Private/TrailHierarchy.cpp
@@ -10,7 +10,6 @@
 #include "Framework/Application/SlateApplication.h"
 #include "SceneManagement.h"
 #include "Tools/MotionTrailOptions.h"
-#include "TrajectoryCache.h"
 #include "Widgets/Input/SButton.h"
 #include "Widgets/Input/SCheckBox.h"
 #include "Widgets/Text/STextBlock.h"
@@ -78,54 +77,54 @@ void FTrail::GetColor(const FFrameNumber& CurrentTime, FColorState& State)
 
 void FTrailHierarchyRenderer::Render(const FSceneView* View, FPrimitiveDrawInterface* PDI)
 {
-	const FDateTime RenderStartTime = FDateTime::Now();
-	
-	const int32 NumEvalTimes = int32((OwningHierarchy->GetViewFrameRange().GetUpperBoundValue().Value - OwningHierarchy->GetViewFrameRange().GetLowerBoundValue().Value) / OwningHierarchy->GetFramesPerFrame().Value) + 1;
-	const int32 NumLinesReserveSize = int32(NumEvalTimes * OwningHierarchy->GetAllTrails().Num() * 1.3);
-	PDI->AddReserveLines(SDPG_Foreground, NumLinesReserveSize);
-
-	TArray<FGuid> TrailsToRemove; //if trail is not visible we just get rid of it
-	// <parent, current>
-	FColorState ColorState;
-	ColorState.Setup(OwningHierarchy);
-	const FCurrentFramesInfo* CurrentFramesInfo = OwningHierarchy->GetCurrentFramesInfo();
-	for (const TPair<FGuid, TUniquePtr<FTrail>>& GuidTrailPair : OwningHierarchy->GetAllTrails())
+	if (CachedOptions)
 	{
-		FTrajectoryDrawInfo* CurDrawInfo = GuidTrailPair.Value->GetDrawInfo();
-		if (CurDrawInfo && OwningHierarchy->GetVisibilityManager().IsTrailVisible(GuidTrailPair.Key, GuidTrailPair.Value.Get()))
+		//options object won't change over lifetime of the renderer functions
+		const FDateTime RenderStartTime = FDateTime::Now();
+
+		const int32 NumEvalTimes = int32((OwningHierarchy->GetViewFrameRange().GetUpperBoundValue().Value - OwningHierarchy->GetViewFrameRange().GetLowerBoundValue().Value) / OwningHierarchy->GetFramesPerFrame().Value) + 1;
+		const int32 NumLinesReserveSize = int32(NumEvalTimes * OwningHierarchy->GetAllTrails().Num() * 1.3);
+		PDI->AddReserveLines(SDPG_Foreground, NumLinesReserveSize);
+
+		TArray<FGuid> TrailsToRemove; //if trail is not visible we just get rid of it
+		// <parent, current>
+		FColorState ColorState;
+		ColorState.Setup(OwningHierarchy);
+		const FCurrentFramesInfo* CurrentFramesInfo = OwningHierarchy->GetCurrentFramesInfo();
+		for (const TPair<FGuid, TUniquePtr<FTrail>>& GuidTrailPair : OwningHierarchy->GetAllTrails())
 		{
-			if (OwningHierarchy->IsTrailEvaluating(GuidTrailPair.Key, true))
+			FTrajectoryDrawInfo* CurDrawInfo = GuidTrailPair.Value->GetDrawInfo();
+			if (CurDrawInfo && OwningHierarchy->GetVisibilityManager().IsTrailVisible(GuidTrailPair.Key, GuidTrailPair.Value.Get(), CachedOptions->bShowSelectedTrails))
 			{
-				GuidTrailPair.Value->RenderEvaluating(GuidTrailPair.Key, View, PDI);
-				continue;
-			}
-			if (GuidTrailPair.Value.Get()->GetCacheState() != ETrailCacheState::UpToDate)
-			{
-				continue;
+				if (OwningHierarchy->IsTrailEvaluating(GuidTrailPair.Key, true))
+				{
+					GuidTrailPair.Value->RenderEvaluating(GuidTrailPair.Key, View, PDI);
+					continue;
+				}
+				if (GuidTrailPair.Value.Get()->GetCacheState() != ETrailCacheState::UpToDate)
+				{
+					continue;
+				}
+				//the above check so that if we are evaluating something that is getting drawn indirectly we skip drawing it, but if it is
+				//getting manipulated we still want to draw it. If it's not evaluating though we don't use sparse time values, instead do every frame.
+				const bool bTrailDirectlyEvaluating = OwningHierarchy->IsTrailEvaluating(GuidTrailPair.Key, false);
+				GuidTrailPair.Value->ReadyToDrawTrail(ColorState, CurrentFramesInfo, bTrailDirectlyEvaluating, OwningHierarchy->GetVisibilityManager().IsTrailAlwaysVisible(GuidTrailPair.Key));
+				GuidTrailPair.Value->Render(GuidTrailPair.Key, View, PDI);
 			}
-			//the above check so that if we are evaluating something that is getting drawn indirectly we skip drawing it, but if it is
-			//getting manipulated we still want to draw it. If it's not evaluating though we don't use sparse time values, instead do every frame.
-			const bool bTrailDirectlyEvaluating = OwningHierarchy->IsTrailEvaluating(GuidTrailPair.Key, false);
-			GuidTrailPair.Value->ReadyToDrawTrail(ColorState, CurrentFramesInfo, bTrailDirectlyEvaluating,  OwningHierarchy->GetVisibilityManager().IsTrailAlwaysVisible(GuidTrailPair.Key));
-			GuidTrailPair.Value->Render(GuidTrailPair.Key,View,PDI);
+			
 		}
-		else
+		for (const FGuid& Key : TrailsToRemove)
 		{
-			TrailsToRemove.Add(GuidTrailPair.Key);
+			OwningHierarchy->RemoveTrail(Key);
 		}
-
-	}
-	for (const FGuid& Key : TrailsToRemove)
-	{
-		OwningHierarchy->RemoveTrail(Key);
+		const FTimespan RenderTimespan = FDateTime::Now() - RenderStartTime;
+		OwningHierarchy->GetTimingStats().Add("FTrailHierarchyRenderer::Render", RenderTimespan);
 	}
-	const FTimespan RenderTimespan = FDateTime::Now() - RenderStartTime;
-	OwningHierarchy->GetTimingStats().Add("FTrailHierarchyRenderer::Render", RenderTimespan);
 }
 
 void FTrailHierarchyRenderer::DrawHUD(const FSceneView* View, FCanvas* Canvas)
 {
-	if (UMotionTrailToolOptions::GetTrailOptions()->bShowMarks)
+	if (CachedOptions && CachedOptions->bShowMarks)
 	{
 		const FDateTime DrawHUDStartTime = FDateTime::Now();
 		const FFrameNumber FramesPerMark = OwningHierarchy->GetFramesPerFrame();
@@ -139,7 +138,7 @@ void FTrailHierarchyRenderer::DrawHUD(const FSceneView* View, FCanvas* Canvas)
 
 		for (const TPair<FGuid, TUniquePtr<FTrail>>& GuidTrailPair : OwningHierarchy->GetAllTrails())
 		{
-			if (OwningHierarchy->GetVisibilityManager().IsTrailVisible(GuidTrailPair.Key, GuidTrailPair.Value.Get()))
+			if (OwningHierarchy->GetVisibilityManager().IsTrailVisible(GuidTrailPair.Key, GuidTrailPair.Value.Get(), CachedOptions->bShowSelectedTrails))
 			{
 				if (OwningHierarchy->IsTrailEvaluating(GuidTrailPair.Key, true) || GuidTrailPair.Value.Get()->GetCacheState() != ETrailCacheState::UpToDate)
 				{
diff --git a/Engine/Plugins/Editor/SequencerAnimTools/Source/SequencerAnimTools/Private/TrajectoryDrawInfo.cpp b/Engine/Plugins/Editor/SequencerAnimTools/Source/SequencerAnimTools/Private/TrajectoryDrawInfo.cpp
--- a/Engine/Plugins/Editor/SequencerAnimTools/Source/SequencerAnimTools/Private/TrajectoryDrawInfo.cpp
+++ b/Engine/Plugins/Editor/SequencerAnimTools/Source/SequencerAnimTools/Private/TrajectoryDrawInfo.cpp
@@ -2,7 +2,6 @@
 
 #include "TrajectoryDrawInfo.h"
 #include "SceneView.h"
-#include "TrajectoryCache.h"
 #include "Framework/Application/SlateApplication.h"
 
 
diff --git a/Engine/Plugins/Editor/SequencerAnimTools/Source/SequencerAnimTools/Public/TrajectoryCache.h b/Engine/Plugins/Editor/SequencerAnimTools/Source/SequencerAnimTools/Public/TrajectoryCache.h
deleted file mode 100644
