CIS Fix.
#trivial

diff --git a/Engine/Plugins/Editor/SequencerAnimTools/Source/SequencerAnimTools/Private/TrailHierarchy.cpp b/Engine/Plugins/Editor/SequencerAnimTools/Source/SequencerAnimTools/Private/TrailHierarchy.cpp
--- a/Engine/Plugins/Editor/SequencerAnimTools/Source/SequencerAnimTools/Private/TrailHierarchy.cpp
+++ b/Engine/Plugins/Editor/SequencerAnimTools/Source/SequencerAnimTools/Private/TrailHierarchy.cpp
@@ -77,6 +77,7 @@ void FTrail::GetColor(const FFrameNumber& CurrentTime, FColorState& State)
 
 void FTrailHierarchyRenderer::Render(const FSceneView* View, FPrimitiveDrawInterface* PDI)
 {
+	/*
 	if (CachedOptions)
 	{
 		//options object won't change over lifetime of the renderer functions
@@ -86,8 +87,6 @@ void FTrailHierarchyRenderer::Render(const FSceneView* View, FPrimitiveDrawInter
 		const int32 NumLinesReserveSize = int32(NumEvalTimes * OwningHierarchy->GetAllTrails().Num() * 1.3);
 		PDI->AddReserveLines(SDPG_Foreground, NumLinesReserveSize);
 
-		TArray<FGuid> TrailsToRemove; //if trail is not visible we just get rid of it
-		// <parent, current>
 		FColorState ColorState;
 		ColorState.Setup(OwningHierarchy);
 		const FCurrentFramesInfo* CurrentFramesInfo = OwningHierarchy->GetCurrentFramesInfo();
@@ -113,17 +112,60 @@ void FTrailHierarchyRenderer::Render(const FSceneView* View, FPrimitiveDrawInter
 			}
 			
 		}
-		for (const FGuid& Key : TrailsToRemove)
-		{
-			OwningHierarchy->RemoveTrail(Key);
-		}
 		const FTimespan RenderTimespan = FDateTime::Now() - RenderStartTime;
 		OwningHierarchy->GetTimingStats().Add("FTrailHierarchyRenderer::Render", RenderTimespan);
 	}
+	*/
+
+	const FDateTime RenderStartTime = FDateTime::Now();
+
+	const int32 NumEvalTimes = int32((OwningHierarchy->GetViewFrameRange().GetUpperBoundValue().Value - OwningHierarchy->GetViewFrameRange().GetLowerBoundValue().Value) / OwningHierarchy->GetFramesPerFrame().Value) + 1;
+	const int32 NumLinesReserveSize = int32(NumEvalTimes * OwningHierarchy->GetAllTrails().Num() * 1.3);
+	PDI->AddReserveLines(SDPG_Foreground, NumLinesReserveSize);
+
+	TArray<FGuid> TrailsToRemove; //if trail is not visible we just get rid of it
+	// <parent, current>
+	FColorState ColorState;
+	ColorState.Setup(OwningHierarchy);
+	const FCurrentFramesInfo* CurrentFramesInfo = OwningHierarchy->GetCurrentFramesInfo();
+	for (const TPair<FGuid, TUniquePtr<FTrail>>& GuidTrailPair : OwningHierarchy->GetAllTrails())
+	{
+		FTrajectoryDrawInfo* CurDrawInfo = GuidTrailPair.Value->GetDrawInfo();
+		if (CurDrawInfo && OwningHierarchy->GetVisibilityManager().IsTrailVisible(GuidTrailPair.Key, GuidTrailPair.Value.Get()))
+		{
+			if (OwningHierarchy->IsTrailEvaluating(GuidTrailPair.Key, true))
+			{
+				GuidTrailPair.Value->RenderEvaluating(GuidTrailPair.Key, View, PDI);
+				continue;
+			}
+			if (GuidTrailPair.Value.Get()->GetCacheState() != ETrailCacheState::UpToDate)
+			{
+				continue;
+			}
+			//the above check so that if we are evaluating something that is getting drawn indirectly we skip drawing it, but if it is
+			//getting manipulated we still want to draw it. If it's not evaluating though we don't use sparse time values, instead do every frame.
+			const bool bTrailDirectlyEvaluating = OwningHierarchy->IsTrailEvaluating(GuidTrailPair.Key, false);
+			GuidTrailPair.Value->ReadyToDrawTrail(ColorState, CurrentFramesInfo, bTrailDirectlyEvaluating, OwningHierarchy->GetVisibilityManager().IsTrailAlwaysVisible(GuidTrailPair.Key));
+			GuidTrailPair.Value->Render(GuidTrailPair.Key, View, PDI);
+		}
+		else
+		{
+			TrailsToRemove.Add(GuidTrailPair.Key);
+		}
+
+	}
+	for (const FGuid& Key : TrailsToRemove)
+	{
+		OwningHierarchy->RemoveTrail(Key);
+	}
+	const FTimespan RenderTimespan = FDateTime::Now() - RenderStartTime;
+	OwningHierarchy->GetTimingStats().Add("FTrailHierarchyRenderer::Render", RenderTimespan);
+
 }
 
 void FTrailHierarchyRenderer::DrawHUD(const FSceneView* View, FCanvas* Canvas)
 {
+	/*
 	if (CachedOptions && CachedOptions->bShowMarks)
 	{
 		const FDateTime DrawHUDStartTime = FDateTime::Now();
@@ -164,6 +206,49 @@ void FTrailHierarchyRenderer::DrawHUD(const FSceneView* View, FCanvas* Canvas)
 
 		}
 
+		const FTimespan DrawHUDTimespan = FDateTime::Now() - DrawHUDStartTime;
+		OwningHierarchy->GetTimingStats().Add("FTrailHierarchyRenderer::DrawHUD", DrawHUDTimespan);
+	}*/
+	if (UMotionTrailToolOptions::GetTrailOptions()->bShowMarks)
+	{
+		const FDateTime DrawHUDStartTime = FDateTime::Now();
+		const FFrameNumber FramesPerMark = OwningHierarchy->GetFramesPerFrame();
+
+		const int32 NumMarksPerTrail = int32((OwningHierarchy->GetViewFrameRange().GetUpperBoundValue().Value - OwningHierarchy->GetViewFrameRange().GetLowerBoundValue().Value) / FramesPerMark.Value) + 1;
+		const int32 PredictedNumMarks = int32(NumMarksPerTrail * OwningHierarchy->GetAllTrails().Num() * 1.3);
+
+		Canvas->GetBatchedElements(FCanvas::EElementType::ET_Line)->AddReserveLines(PredictedNumMarks);
+		FTrailScreenSpaceTransform Transform(View, Canvas->GetDPIScale());
+		const FCurrentFramesInfo* CurrentFramesInfo = OwningHierarchy->GetCurrentFramesInfo();
+
+		for (const TPair<FGuid, TUniquePtr<FTrail>>& GuidTrailPair : OwningHierarchy->GetAllTrails())
+		{
+			if (OwningHierarchy->GetVisibilityManager().IsTrailVisible(GuidTrailPair.Key, GuidTrailPair.Value.Get()))
+			{
+				if (OwningHierarchy->IsTrailEvaluating(GuidTrailPair.Key, true) || GuidTrailPair.Value.Get()->GetCacheState() != ETrailCacheState::UpToDate)
+				{
+					continue;
+				}
+				TArray<FVector2D> Marks, MarkNormals;
+				const bool bTrailDirectlyEvaluating = OwningHierarchy->IsTrailEvaluating(GuidTrailPair.Key, false);
+				GuidTrailPair.Value->GetTickPointsForDisplay(Transform, *CurrentFramesInfo, bTrailDirectlyEvaluating, Marks, MarkNormals);
+				if (Marks.Num() > 0)
+				{
+					const FLinearColor Color = UMotionTrailToolOptions::GetTrailOptions()->MarkColor;
+					for (int32 Idx = 0; Idx < Marks.Num(); Idx++)
+					{
+						const FVector2D StartPoint = Marks[Idx] - MarkNormals[Idx] * UMotionTrailToolOptions::GetTrailOptions()->MarkSize;
+						const FVector2D EndPoint = Marks[Idx] + MarkNormals[Idx] * UMotionTrailToolOptions::GetTrailOptions()->MarkSize;
+						FCanvasLineItem LineItem = FCanvasLineItem(StartPoint, EndPoint);
+						LineItem.SetColor(Color);
+						Canvas->DrawItem(LineItem);
+					}
+				}
+			}
+			GuidTrailPair.Value->DrawHUD(View, Canvas);
+
+		}
+
 		const FTimespan DrawHUDTimespan = FDateTime::Now() - DrawHUDStartTime;
 		OwningHierarchy->GetTimingStats().Add("FTrailHierarchyRenderer::DrawHUD", DrawHUDTimespan);
 	}
