Motion Trail:  Fix crash when drawing hud since trail may not be udpated do to how the hud draws. Clean pin state in UI when using modifiers.

diff --git a/Engine/Plugins/Editor/SequencerAnimTools/Source/SequencerAnimTools/Private/Sequencer/SequencerTrailHierarchy.cpp b/Engine/Plugins/Editor/SequencerAnimTools/Source/SequencerAnimTools/Private/Sequencer/SequencerTrailHierarchy.cpp
--- a/Engine/Plugins/Editor/SequencerAnimTools/Source/SequencerAnimTools/Private/Sequencer/SequencerTrailHierarchy.cpp
+++ b/Engine/Plugins/Editor/SequencerAnimTools/Source/SequencerAnimTools/Private/Sequencer/SequencerTrailHierarchy.cpp
@@ -1050,7 +1050,18 @@ void FSequencerTrailHierarchy::OnUnPinSelection()
 		if (VisibilityManager.IsTrailAlwaysVisible(GuidTrailPair.Key) == true 
 			&& (VisibilityManager.ControlSelected.Contains(GuidTrailPair.Key) || VisibilityManager.Selected.Contains(GuidTrailPair.Key)))
 		{
-			VisibilityManager.SetTrailAlwaysVisible(GuidTrailPair.Key, false);
+			if (UMotionTrailToolOptions* Settings = UMotionTrailToolOptions::GetTrailOptions())
+			{
+				int32 Index = Settings->GetIndexFromGuid(GuidTrailPair.Key);
+				if (Index != INDEX_NONE)
+				{
+					Settings->DeletePinned(Index);
+				}
+			}
+			else
+			{
+				VisibilityManager.SetTrailAlwaysVisible(GuidTrailPair.Key, false);
+			}
 		}
 	}
 }
diff --git a/Engine/Plugins/Editor/SequencerAnimTools/Source/SequencerAnimTools/Private/TrajectoryDrawInfo.cpp b/Engine/Plugins/Editor/SequencerAnimTools/Source/SequencerAnimTools/Private/TrajectoryDrawInfo.cpp
--- a/Engine/Plugins/Editor/SequencerAnimTools/Source/SequencerAnimTools/Private/TrajectoryDrawInfo.cpp
+++ b/Engine/Plugins/Editor/SequencerAnimTools/Source/SequencerAnimTools/Private/TrajectoryDrawInfo.cpp
@@ -205,6 +205,10 @@ void FCurrentFramesInfo::Reset()
 **********************************************************************************/
 void FTrajectoryDrawInfo::GetTrajectoryPointsForDisplay(const FTransform& OffsetTransform, const FTransform& ParentSpaceTransform, const FCurrentFramesInfo& InCurrentFramesInfo, bool bIsEvaluating,  TArray<FVector>& OutPoints, TArray<FFrameNumber>&  OutFrames)
 {
+	if (ArrayOfTransforms.IsValid() == false || ArrayOfTransforms->Transforms.Num() == 0 || InCurrentFramesInfo.TransformIndices.Num() == 0 || InCurrentFramesInfo.TransformIndices.Num() != InCurrentFramesInfo.CurrentFrames.Num())
+	{
+		return;
+	}
 	//if not evaluating we just sample the times that are in range
 	//if we are evaluating then we use the InCurrentFramesInfo.TransformIndices to get the sparse values
 	if (bIsEvaluating == false )
@@ -269,7 +273,7 @@ void FTrajectoryDrawInfo::GetTrajectoryPointsForDisplay(const FTransform& Offset
 
 void FTrajectoryDrawInfo::GetTickPointsForDisplay(const FTransform& OffsetTransform, const FTransform& ParentSpaceTransform, const FTrailScreenSpaceTransform& InScreenSpaceTransform, const FCurrentFramesInfo& InCurrentFramesInfo, bool bIsEvaluating, TArray<FVector2D>& Ticks, TArray<FVector2D>& TickNormals)
 {
-	if (ArrayOfTransforms.IsValid() == false || ArrayOfTransforms->Transforms.Num() == 0 || InCurrentFramesInfo.TransformIndices.Num() != InCurrentFramesInfo.CurrentFrames.Num())
+	if (ArrayOfTransforms.IsValid() == false || ArrayOfTransforms->Transforms.Num() == 0 || InCurrentFramesInfo.TransformIndices.Num() == 0 ||  InCurrentFramesInfo.TransformIndices.Num() != InCurrentFramesInfo.CurrentFrames.Num())
 	{
 		return;
 	}
@@ -299,11 +303,15 @@ void FTrajectoryDrawInfo::GetTickPointsForDisplay(const FTransform& OffsetTransf
 
 FVector FTrajectoryDrawInfo::GetPoint(const FTransform& OffsetTransform, const FTransform& ParentSpaceTransform, const FCurrentFramesInfo& InCurrentFramesInfo, const FFrameNumber& InTime)
 {
-	FTransform Transform = ArrayOfTransforms->Interp(InTime, InCurrentFramesInfo.TransformIndices, InCurrentFramesInfo.CurrentFrames);
-	FVector Position = ParentSpaceTransform.TransformPosition(Transform.GetLocation());
-	Position = OffsetTransform.TransformPosition(Position);
+	if (InCurrentFramesInfo.TransformIndices.Num() > 0 && (InCurrentFramesInfo.TransformIndices.Num() == InCurrentFramesInfo.CurrentFrames.Num()))
+	{
+		FTransform Transform = ArrayOfTransforms->Interp(InTime, InCurrentFramesInfo.TransformIndices, InCurrentFramesInfo.CurrentFrames);
+		FVector Position = ParentSpaceTransform.TransformPosition(Transform.GetLocation());
+		Position = OffsetTransform.TransformPosition(Position);
 
-	return Position;
+		return Position;
+	}
+	return FVector(0.0,0.0,0.0);
 }
 
 } // namespace SequencerAnimTools
