Anim In Engine: Motion Trail crash
#trivial

diff --git a/Engine/Plugins/Editor/SequencerAnimTools/Source/SequencerAnimTools/Private/Sequencer/MovieSceneTransformTrail.cpp b/Engine/Plugins/Editor/SequencerAnimTools/Source/SequencerAnimTools/Private/Sequencer/MovieSceneTransformTrail.cpp
--- a/Engine/Plugins/Editor/SequencerAnimTools/Source/SequencerAnimTools/Private/Sequencer/MovieSceneTransformTrail.cpp
+++ b/Engine/Plugins/Editor/SequencerAnimTools/Source/SequencerAnimTools/Private/Sequencer/MovieSceneTransformTrail.cpp
@@ -24,6 +24,9 @@
 #include "PrimitiveDrawInterface.h"
 #include "Components/SkeletalMeshComponent.h"
 #include "EditMode/ControlRigEditMode.h"
+#include "UnrealEdGlobals.h"
+#include "EditorModeManager.h"
+
 
 #define LOCTEXT_NAMESPACE "SequencerAnimTools"
 
@@ -128,6 +131,11 @@ FMovieSceneTransformTrail::FMovieSceneTransformTrail(const FGuid& InBindingID, U
 	DrawInfo = MakeUnique<FTrajectoryDrawInfo>(UMotionTrailToolOptions::GetTrailOptions()->TrailStyle, UMotionTrailToolOptions::GetTrailOptions()->DefaultColor, ArrayOfTransforms, ParentSpace.ParentSpace.WorldTransforms);
 }
 
+FMovieSceneTransformTrail::~FMovieSceneTransformTrail()
+{
+	ExitOffsetMode();
+}
+
 void FMovieSceneTransformTrail::HasStartedEvaluating()
 {
 	PreviousCachedDrawData = CachedDrawData;
@@ -536,20 +544,23 @@ void FMovieSceneTransformTrail::ReadyToDrawTrail(FColorState& ColorState, const
 	bool bCalculateColor = (ColorState.GetStyle() == EMotionTrailTrailStyle::Time);
 	GetParentSpaceTransform();
 
-	if (InCurrentFramesInfo->bViewRangeIsEvalRange == false || (CachedDrawData.PointsToDraw.Num() != CachedDrawData.Color.Num()) || (CachedDrawData.PointsToDraw.Num() != InCurrentFramesInfo->CurrentFrames.Num())
-		|| (CachedDrawData.Color.Num() != InCurrentFramesInfo->CurrentFrames.Num()))
-	{
-		GetTrajectoryPointsForDisplay(*InCurrentFramesInfo, bIsEvaluating, CachedDrawData.PointsToDraw,CachedDrawData.Frames);
-		bCalculateColor = true;
-	}
-	if (CachedDrawData.PointsToDraw.Num() > 1 && bCalculateColor)
+	if (InCurrentFramesInfo->CurrentFrames.Num() > 0)
 	{
-		CachedDrawData.Color.SetNum(CachedDrawData.PointsToDraw.Num());
-		for (int32 Idx = 1; Idx < CachedDrawData.PointsToDraw.Num(); Idx++)
+		if (InCurrentFramesInfo->bViewRangeIsEvalRange == false || (CachedDrawData.PointsToDraw.Num() != CachedDrawData.Color.Num()) || (CachedDrawData.PointsToDraw.Num() != InCurrentFramesInfo->CurrentFrames.Num())
+			|| (CachedDrawData.Color.Num() != InCurrentFramesInfo->CurrentFrames.Num()))
 		{
-			const FFrameNumber CurFrame = CachedDrawData.Frames[Idx - 1];
-			GetColor(CurFrame, ColorState);
-			CachedDrawData.Color[Idx - 1] = ColorState.CalculatedColor;
+			GetTrajectoryPointsForDisplay(*InCurrentFramesInfo, bIsEvaluating, CachedDrawData.PointsToDraw, CachedDrawData.Frames);
+			bCalculateColor = true;
+		}
+		if (CachedDrawData.PointsToDraw.Num() > 1 && bCalculateColor)
+		{
+			CachedDrawData.Color.SetNum(CachedDrawData.PointsToDraw.Num());
+			for (int32 Idx = 1; Idx < CachedDrawData.PointsToDraw.Num(); Idx++)
+			{
+				const FFrameNumber CurFrame = CachedDrawData.Frames[Idx - 1];
+				GetColor(CurFrame, ColorState);
+				CachedDrawData.Color[Idx - 1] = ColorState.CalculatedColor;
+			}
 		}
 	}
 }
@@ -629,7 +640,7 @@ bool FMovieSceneTransformTrail::HandleAltClick(FEditorViewportClient* InViewport
 
 bool FMovieSceneTransformTrail::HandleClick(const FGuid& Guid, FEditorViewportClient* InViewportClient, HHitProxy* InHitProxy, FInputClick Click)
 {
-	bIsOffsetMode = false;
+	ExitOffsetMode();
 
 	bool bKeySelected = KeyTool->HandleClick(Guid,InViewportClient, InHitProxy, Click);
 	if (bKeySelected)
@@ -637,6 +648,7 @@ bool FMovieSceneTransformTrail::HandleClick(const FGuid& Guid, FEditorViewportCl
 		if (Click.bShiftIsDown == false || Click.bCtrlIsDown == false)
 		{
 			bIsSelected = false;
+			ExitOffsetMode();
 		}
 		return true;
 	}
@@ -661,6 +673,7 @@ bool FMovieSceneTransformTrail::HandleClick(const FGuid& Guid, FEditorViewportCl
 				if (bIsSelected)
 				{
 					bIsSelected = false;
+					ExitOffsetMode();
 				}
 				else
 				{
@@ -679,6 +692,7 @@ bool FMovieSceneTransformTrail::HandleClick(const FGuid& Guid, FEditorViewportCl
 	}
 
 	bIsSelected = false;
+	ExitOffsetMode();
 	return false;
 }
 
@@ -735,6 +749,8 @@ void FMovieSceneTransformTrail::SelectNone()
 {
 	KeyTool->ClearSelection();
 	bIsSelected = false;
+	ExitOffsetMode();
+
 }
 
 bool FMovieSceneTransformTrail::IsAnythingSelected() const
@@ -746,14 +762,22 @@ bool FMovieSceneTransformTrail::IsTrailSelected() const
 {
 	return bIsSelected;
 }
-
+void FMovieSceneTransformTrail::ExitOffsetMode()
+{
+	bIsOffsetMode = false;
+	if (bSetShowWidget.IsSet() && bSetShowWidget.GetValue() == true)
+	{
+		GLevelEditorModeTools().SetShowWidget(true); //turn widget back on
+		bSetShowWidget.Reset();
+	}
+}
 void FMovieSceneTransformTrail::ClearOffsetTransform()
 {
 	FTrail::ClearOffsetTransform();
 	OffsetTransform = FTransform::Identity;
 	if (bIsOffsetMode)
 	{
-		bIsOffsetMode = false;
+		ExitOffsetMode();
 		bIsSelected = false;
 		ClearCachedData();
 	}
@@ -763,8 +787,13 @@ void FMovieSceneTransformTrail::SetOffsetMode()
 {
 	if (bIsOffsetMode == false)
 	{
-		bIsOffsetMode = true;
 		SelectNone();
+		if (GLevelEditorModeTools().GetShowWidget())
+		{
+			bSetShowWidget = true;
+			GLevelEditorModeTools().SetShowWidget(false);
+		}
+		bIsOffsetMode = true;
 		bIsSelected = true;
 		//set selected position aka gizmo location  to be location of the current object
 		FFrameTime StartTime = GetSequencer()->GetLocalTime().Time;
@@ -782,13 +811,14 @@ bool FMovieSceneTransformTrail::IsTracking() const
 bool FMovieSceneTransformTrail::StartTracking()
 {
 	bIsTracking = true;
+	/* not sure we want this
 	//if in offset mode this happened from the tool bar menu,
 	// if not we check fo shift is held we go into offset mode
 	if (!bIsOffsetMode)
 	{
 		bIsOffsetMode = FSlateApplication::Get().GetModifierKeys().IsShiftDown();
 	}
-	
+	*/
 	return false;
 }
 
@@ -1216,7 +1246,7 @@ bool FMovieSceneSocketTransformTrail::HandleClick(const FGuid& Guid, FEditorView
 				if (bIsSelected)
 				{
 					bIsSelected = false;
-
+					ExitOffsetMode();
 				}
 				else
 				{
@@ -1234,6 +1264,7 @@ bool FMovieSceneSocketTransformTrail::HandleClick(const FGuid& Guid, FEditorView
 		}
 	}
 	bIsSelected = false;
+	ExitOffsetMode();
 	return false;
 }
 
diff --git a/Engine/Plugins/Editor/SequencerAnimTools/Source/SequencerAnimTools/Private/Sequencer/MovieSceneTransformTrail.h b/Engine/Plugins/Editor/SequencerAnimTools/Source/SequencerAnimTools/Private/Sequencer/MovieSceneTransformTrail.h
--- a/Engine/Plugins/Editor/SequencerAnimTools/Source/SequencerAnimTools/Private/Sequencer/MovieSceneTransformTrail.h
+++ b/Engine/Plugins/Editor/SequencerAnimTools/Source/SequencerAnimTools/Private/Sequencer/MovieSceneTransformTrail.h
@@ -67,6 +67,7 @@ class FMovieSceneTransformTrail : public FTrail
 public:
 
 	FMovieSceneTransformTrail(const FGuid& InBindingID, USceneComponent* SceneComponent, const bool bInIsVisible, TWeakObjectPtr<UMovieSceneTrack> InWeakTrack, TSharedPtr<ISequencer> InSequencer);
+	~FMovieSceneTransformTrail();
 
 
 	// Begin FTrail interface
@@ -143,6 +144,8 @@ protected:
 	bool bIsOffsetMode;
 	FVector SelectedPos;
 	bool bIsTracking; 
+	TOptional<bool> bSetShowWidget;
+
 	//main tool for adjusting keys
 	TUniquePtr<FMotionTraiMovieScenelKeyTool> KeyTool;
 
@@ -152,7 +155,8 @@ protected:
 	FGuid CachedHierarchyGuid;
 protected:
 
-	bool TrailOrSpaceHasChanged() ;
+	bool TrailOrSpaceHasChanged();
+	void ExitOffsetMode();
 	bool BindingHasChanged(const FGuid& InBinding, const USceneComponent* InComponent, UE::AIE::FSequencerTransformDependencies& InDependencies);
 
 protected: 
