Motion Trails: Fix for default color

diff --git a/Engine/Plugins/Editor/SequencerAnimTools/Source/SequencerAnimTools/Private/TrailHierarchy.cpp b/Engine/Plugins/Editor/SequencerAnimTools/Source/SequencerAnimTools/Private/TrailHierarchy.cpp
--- a/Engine/Plugins/Editor/SequencerAnimTools/Source/SequencerAnimTools/Private/TrailHierarchy.cpp
+++ b/Engine/Plugins/Editor/SequencerAnimTools/Source/SequencerAnimTools/Private/TrailHierarchy.cpp
@@ -68,14 +68,8 @@ void FTrail::GetColor(const FFrameNumber& CurrentTime, FColorState& State)
 	}
 	else
 	{
-		if (State.bIsPinned && DrawInfo.IsValid())
-		{
-			State.CalculatedColor = DrawInfo->GetColor();
-		}
-		else
-		{
-			State.CalculatedColor = State.Options->DefaultColor;
-		}
+		State.CalculatedColor = State.Options->DefaultColor;
+		
 	}
 }
 
