Motion Trail: Speculative fix for motion trail crashes

diff --git a/Engine/Plugins/Editor/SequencerAnimTools/Source/SequencerAnimTools/Private/TrailHierarchy.cpp b/Engine/Plugins/Editor/SequencerAnimTools/Source/SequencerAnimTools/Private/TrailHierarchy.cpp
--- a/Engine/Plugins/Editor/SequencerAnimTools/Source/SequencerAnimTools/Private/TrailHierarchy.cpp
+++ b/Engine/Plugins/Editor/SequencerAnimTools/Source/SequencerAnimTools/Private/TrailHierarchy.cpp
@@ -82,7 +82,11 @@ void FTrailHierarchyRenderer::Render(const FSceneView* View, FPrimitiveDrawInter
 		const FDateTime RenderStartTime = FDateTime::Now();
 
 		const int32 NumEvalTimes = int32((OwningHierarchy->GetViewFrameRange().GetUpperBoundValue().Value - OwningHierarchy->GetViewFrameRange().GetLowerBoundValue().Value) / OwningHierarchy->GetFramesPerFrame().Value) + 1;
-		const int32 NumLinesReserveSize = int32(NumEvalTimes * OwningHierarchy->GetAllTrails().Num() * 1.3);
+		if (NumEvalTimes <= 1)
+		{
+			return;
+		}
+		int32 NumLinesReserveSize = int32(NumEvalTimes * OwningHierarchy->GetAllTrails().Num() * 1.3);
 		PDI->AddReserveLines(SDPG_Foreground, NumLinesReserveSize);
 
 		FColorState ColorState;
@@ -124,6 +128,10 @@ void FTrailHierarchyRenderer::DrawHUD(const FSceneView* View, FCanvas* Canvas)
 		const FFrameNumber FramesPerMark = OwningHierarchy->GetFramesPerFrame();
 
 		const int32 NumMarksPerTrail = int32((OwningHierarchy->GetViewFrameRange().GetUpperBoundValue().Value - OwningHierarchy->GetViewFrameRange().GetLowerBoundValue().Value) / FramesPerMark.Value) + 1;
+		if (NumMarksPerTrail <= 1)
+		{
+			return;
+		}
 		const int32 PredictedNumMarks = int32(NumMarksPerTrail * OwningHierarchy->GetAllTrails().Num() * 1.3);
 
 		Canvas->GetBatchedElements(FCanvas::EElementType::ET_Line)->AddReserveLines(PredictedNumMarks);
