Motion trails; Use subsequence range

diff --git a/Engine/Plugins/Editor/SequencerAnimTools/Source/SequencerAnimTools/Private/Sequencer/SequencerTrailHierarchy.cpp b/Engine/Plugins/Editor/SequencerAnimTools/Source/SequencerAnimTools/Private/Sequencer/SequencerTrailHierarchy.cpp
--- a/Engine/Plugins/Editor/SequencerAnimTools/Source/SequencerAnimTools/Private/Sequencer/SequencerTrailHierarchy.cpp
+++ b/Engine/Plugins/Editor/SequencerAnimTools/Source/SequencerAnimTools/Private/Sequencer/SequencerTrailHierarchy.cpp
@@ -757,7 +757,8 @@ void FSequencerTrailHierarchy::UpdateViewAndEvalRange()
 
 	FFrameRate TickResolution = Sequencer->GetFocusedTickResolution();
 	FFrameRate DisplayRate = Sequencer->GetFocusedDisplayRate();
-	TickEvalRange  = Sequencer->GetFocusedMovieSceneSequence()->GetMovieScene()->GetPlaybackRange();
+	TOptional<TRange<FFrameNumber>> OptionalRange = Sequencer->GetSubSequenceRange();
+	TickEvalRange  = OptionalRange.IsSet()  ? OptionalRange.GetValue() : Sequencer->GetFocusedMovieSceneSequence()->GetMovieScene()->GetPlaybackRange();
 	if (!UMotionTrailToolOptions::GetTrailOptions()->bShowFullTrail)
 	{
 		FFrameTime SequenceTime = Sequencer->GetLocalTime().Time;
