Anim In Engine: Fix Pivot square scale in different FOV's.

diff --git a/Engine/Plugins/Editor/SequencerAnimTools/Source/SequencerAnimTools/Private/SequencerAnimEditPivotTool.cpp b/Engine/Plugins/Editor/SequencerAnimTools/Source/SequencerAnimTools/Private/SequencerAnimEditPivotTool.cpp
--- a/Engine/Plugins/Editor/SequencerAnimTools/Source/SequencerAnimTools/Private/SequencerAnimEditPivotTool.cpp
+++ b/Engine/Plugins/Editor/SequencerAnimTools/Source/SequencerAnimTools/Private/SequencerAnimEditPivotTool.cpp
@@ -1036,7 +1036,17 @@ void USequencerPivotTool::Render(IToolsContextRenderAPI* RenderAPI)
 		FPrimitiveDrawInterface* PDI = RenderAPI->GetPrimitiveDrawInterface();
 
 		const bool bHitTesting = PDI && PDI->IsHitTesting();
-		const float KeySize = 20.0f;
+		float KeySize = 20.0f;
+		//hack for PDI->DrawPoint not scaling perspective keys see FViewElementPDI::DrawPoint
+		if (RenderAPI && RenderAPI->GetSceneView())
+		{
+			const bool bIsPerspective = (RenderAPI->GetSceneView()->ViewMatrices.GetProjectionMatrix().M[3][3] < 1.0f) ? true : false;
+			if (bIsPerspective)
+			{
+				const float ZoomFactor = FMath::Min<float>(RenderAPI->GetSceneView()->ViewMatrices.GetProjectionMatrix().M[0][0], RenderAPI->GetSceneView()->ViewMatrices.GetProjectionMatrix().M[1][1]);
+				KeySize = KeySize / ZoomFactor;
+			}
+		}
 		const FLinearColor Color = FLinearColor(1.0f, 0.0f, 0.0f);
 		if (bHitTesting)
 		{
