Anim In Engine: Render checker for Fix Pivot square scale in different FOV's.

diff --git a/Engine/Plugins/Editor/SequencerAnimTools/Source/SequencerAnimTools/Private/SequencerAnimEditPivotTool.cpp b/Engine/Plugins/Editor/SequencerAnimTools/Source/SequencerAnimTools/Private/SequencerAnimEditPivotTool.cpp
--- a/Engine/Plugins/Editor/SequencerAnimTools/Source/SequencerAnimTools/Private/SequencerAnimEditPivotTool.cpp
+++ b/Engine/Plugins/Editor/SequencerAnimTools/Source/SequencerAnimTools/Private/SequencerAnimEditPivotTool.cpp
@@ -1031,6 +1031,10 @@ void USequencerPivotTool::OnPropertyModified(UObject* PropertySet, FProperty* Pr
 
 void USequencerPivotTool::Render(IToolsContextRenderAPI* RenderAPI)
 {
+	if (RenderAPI == nullptr)
+	{
+		return;
+	}
 	if (bPickingPivotLocation == false )
 	{
 		FPrimitiveDrawInterface* PDI = RenderAPI->GetPrimitiveDrawInterface();
@@ -1038,7 +1042,7 @@ void USequencerPivotTool::Render(IToolsContextRenderAPI* RenderAPI)
 		const bool bHitTesting = PDI && PDI->IsHitTesting();
 		float KeySize = 20.0f;
 		//hack for PDI->DrawPoint not scaling perspective keys see FViewElementPDI::DrawPoint
-		if (RenderAPI && RenderAPI->GetSceneView())
+		if (RenderAPI->GetSceneView())
 		{
 			const bool bIsPerspective = (RenderAPI->GetSceneView()->ViewMatrices.GetProjectionMatrix().M[3][3] < 1.0f) ? true : false;
 			if (bIsPerspective)
