Anim In Engine: Motion trails can be selected without a focused viewport, simple fix to just add a bunch of nullprt checks

diff --git a/Engine/Plugins/Editor/SequencerAnimTools/Source/SequencerAnimTools/Private/MotionTrailTool.cpp b/Engine/Plugins/Editor/SequencerAnimTools/Source/SequencerAnimTools/Private/MotionTrailTool.cpp
--- a/Engine/Plugins/Editor/SequencerAnimTools/Source/SequencerAnimTools/Private/MotionTrailTool.cpp
+++ b/Engine/Plugins/Editor/SequencerAnimTools/Source/SequencerAnimTools/Private/MotionTrailTool.cpp
@@ -220,14 +220,19 @@ void UMotionTrailTool::OnUpdateModifierState(int ModifierID, bool bIsOn)
 FInputRayHit UMotionTrailTool::IsHitByClick(const FInputDeviceRay& ClickPos)
 {
 	FInputRayHit ReturnHit = FInputRayHit();
-	FViewport* FocusedViewport = GetToolManager()->GetContextQueriesAPI()->GetFocusedViewport();
-	HHitProxy* HitResult = FocusedViewport->GetHitProxy(ClickPos.ScreenPosition.X, ClickPos.ScreenPosition.Y);
-	for (TUniquePtr<UE::SequencerAnimTools::FTrailHierarchy>& TrailHierarchy : TrailHierarchies)
+	if (GetToolManager() && GetToolManager()->GetContextQueriesAPI())
 	{
-		if (TrailHierarchy->IsHitByClick(HitResult))
+		if (FViewport* FocusedViewport = GetToolManager()->GetContextQueriesAPI()->GetFocusedViewport())
 		{
-			ReturnHit.bHit = true;
-			break;
+			HHitProxy* HitResult = FocusedViewport->GetHitProxy(ClickPos.ScreenPosition.X, ClickPos.ScreenPosition.Y);
+			for (TUniquePtr<UE::SequencerAnimTools::FTrailHierarchy>& TrailHierarchy : TrailHierarchies)
+			{
+				if (TrailHierarchy->IsHitByClick(HitResult))
+				{
+					ReturnHit.bHit = true;
+					break;
+				}
+			}
 		}
 	}
 
