Motion Trails: Make sure Motion trail state stays on with asset reload. The bShowTrails flag is handled elsewhere.

diff --git a/Engine/Plugins/Editor/SequencerAnimTools/Source/SequencerAnimTools/Private/MotionTrailTool.cpp b/Engine/Plugins/Editor/SequencerAnimTools/Source/SequencerAnimTools/Private/MotionTrailTool.cpp
--- a/Engine/Plugins/Editor/SequencerAnimTools/Source/SequencerAnimTools/Private/MotionTrailTool.cpp
+++ b/Engine/Plugins/Editor/SequencerAnimTools/Source/SequencerAnimTools/Private/MotionTrailTool.cpp
@@ -172,10 +172,6 @@ void UMotionTrailTool::SetupIntegration()
 
 void UMotionTrailTool::Shutdown(EToolShutdownType ShutdownType)
 {
-	if (UMotionTrailToolOptions::GetTrailOptions()->bShowTrails == true)
-	{
-		UMotionTrailToolOptions::GetTrailOptions()->bShowTrails = false;
-	}
 	Super::Shutdown(ShutdownType);
 	GizmoManager->DestroyAllGizmosByOwner(this);
 	ShutdownIntegration();
