Motion Trails: Fix for crash reporter crash, add some weak pointers to lambdas

diff --git a/Engine/Plugins/Editor/SequencerAnimTools/Source/SequencerAnimTools/Private/Sequencer/SequencerTrailHierarchy.cpp b/Engine/Plugins/Editor/SequencerAnimTools/Source/SequencerAnimTools/Private/Sequencer/SequencerTrailHierarchy.cpp
--- a/Engine/Plugins/Editor/SequencerAnimTools/Source/SequencerAnimTools/Private/Sequencer/SequencerTrailHierarchy.cpp
+++ b/Engine/Plugins/Editor/SequencerAnimTools/Source/SequencerAnimTools/Private/Sequencer/SequencerTrailHierarchy.cpp
@@ -882,11 +882,30 @@ void FSequencerTrailHierarchy::ClearSelection()
 
 void FSequencerTrailHierarchy::RegisterControlRigDelegates(USkeletalMeshComponent* Component, UMovieSceneControlRigParameterTrack* CRParameterTrack)
 {
+	if (CRParameterTrack == nullptr)
+	{
+		return;
+	}
 	UControlRig* ControlRig = CRParameterTrack->GetControlRig();
-	URigHierarchy* RigHierarchy = ControlRig->GetHierarchy();
+	if (Component == nullptr || ControlRig == nullptr || ControlRig->GetHierarchy() == nullptr)
+	{
+		return;
+	}
+	TWeakObjectPtr<USkeletalMeshComponent> WeakComponent = Component;
+	TWeakObjectPtr<UMovieSceneControlRigParameterTrack> WeakTrack = CRParameterTrack;
+
 	FControlRigDelegateHandles& DelegateHandles = ControlRigDelegateHandles.Add(CRParameterTrack);
-	DelegateHandles.OnControlSelected = ControlRig->ControlSelected().AddLambda([this, Component,CRParameterTrack](UControlRig* ControlRig, FRigControlElement* ControlElement, bool bSelected)
+
+	DelegateHandles.OnControlSelected = ControlRig->ControlSelected().AddLambda([this, WeakComponent, WeakTrack](UControlRig* ControlRig, FRigControlElement* ControlElement, bool bSelected)
 		{
+			TStrongObjectPtr<USkeletalMeshComponent> ComponentPtr = WeakComponent.Pin();
+			TStrongObjectPtr<UMovieSceneControlRigParameterTrack> TrackPtr = WeakTrack.Pin();
+			TSharedPtr <ISequencer> SequencerPtr = WeakSequencer.Pin();
+			if (!ComponentPtr || !TrackPtr || !SequencerPtr)
+			{
+				return;
+			}
+				
 			if (ControlElement->Settings.ControlType != ERigControlType::Transform &&
 				ControlElement->Settings.ControlType != ERigControlType::TransformNoScale &&
 				ControlElement->Settings.ControlType != ERigControlType::EulerTransform)
@@ -896,7 +915,7 @@ void FSequencerTrailHierarchy::RegisterControlRigDelegates(USkeletalMeshComponen
 
 			if (bSelected)
 			{
-				AddControlRigTrail(Component,ControlRig,CRParameterTrack, ControlElement->GetFName());
+				AddControlRigTrail(ComponentPtr.Get(), ControlRig, TrackPtr.Get(), ControlElement->GetFName());
 			}
 
 			if (ControlsTracked.Find(ControlRig) != nullptr)
@@ -938,10 +957,18 @@ void FSequencerTrailHierarchy::RegisterControlRigDelegates(USkeletalMeshComponen
 			
 	);
 	
+	//just use ControlRig and not WeakControlRig since it's only used as a TMap Key
+	URigHierarchy* RigHierarchy = ControlRig->GetHierarchy();
 	DelegateHandles.OnHierarchyModified = RigHierarchy->OnModified().AddLambda(
-		[this, RigHierarchy, ControlRig](ERigHierarchyNotification InNotif, URigHierarchy* InHierarchy, const FRigNotificationSubject& InSubject) {
+		[this, ControlRig](ERigHierarchyNotification InNotif, URigHierarchy* InHierarchy, const FRigNotificationSubject& InSubject) 
+		{
+		
 			TSharedPtr<ISequencer> Sequencer = WeakSequencer.Pin();
-			check(Sequencer);
+			if (!Sequencer == false)
+			{
+				return;
+			}
+			
 
 			const FRigControlElement* ControlElement = Cast<FRigControlElement>(InSubject.Element); 
 			if(ControlElement == nullptr)
