Anim In Engine: Motion Trails, fix strobing frame coloring with dashed by setting the start frame to be eval range start so it doesn't change when view range changes.

diff --git a/Engine/Plugins/Editor/SequencerAnimTools/Source/SequencerAnimTools/Private/Sequencer/MovieSceneTransformTrail.cpp b/Engine/Plugins/Editor/SequencerAnimTools/Source/SequencerAnimTools/Private/Sequencer/MovieSceneTransformTrail.cpp
--- a/Engine/Plugins/Editor/SequencerAnimTools/Source/SequencerAnimTools/Private/Sequencer/MovieSceneTransformTrail.cpp
+++ b/Engine/Plugins/Editor/SequencerAnimTools/Source/SequencerAnimTools/Private/Sequencer/MovieSceneTransformTrail.cpp
@@ -541,7 +541,7 @@ void FMovieSceneTransformTrail::GetColor(const FFrameNumber& CurrentTime, FColor
 void FMovieSceneTransformTrail::ReadyToDrawTrail(FColorState& ColorState, const FCurrentFramesInfo* InCurrentFramesInfo, bool bIsEvaluating, bool bIsPinned)
 {
 	ColorState.ReadyForTrail(bIsPinned, GetDrawInfo()->GetStyle());
-	bool bCalculateColor = (ColorState.GetStyle() == EMotionTrailTrailStyle::Time);
+	bool bCalculateColor = (ColorState.GetStyle() == EMotionTrailTrailStyle::Time || ColorState.GetStyle() == EMotionTrailTrailStyle::Dashed);
 	GetParentSpaceTransform();
 
 	if (InCurrentFramesInfo->CurrentFrames.Num() > 0)
diff --git a/Engine/Plugins/Editor/SequencerAnimTools/Source/SequencerAnimTools/Private/Sequencer/SequencerTrailHierarchy.cpp b/Engine/Plugins/Editor/SequencerAnimTools/Source/SequencerAnimTools/Private/Sequencer/SequencerTrailHierarchy.cpp
--- a/Engine/Plugins/Editor/SequencerAnimTools/Source/SequencerAnimTools/Private/Sequencer/SequencerTrailHierarchy.cpp
+++ b/Engine/Plugins/Editor/SequencerAnimTools/Source/SequencerAnimTools/Private/Sequencer/SequencerTrailHierarchy.cpp
@@ -38,6 +38,7 @@ void FColorState::Setup(FTrailHierarchy* TrailHierarchy)
 		FSequencerTrailHierarchy* SequencerTrailHierarchy = static_cast<FSequencerTrailHierarchy*>(TrailHierarchy);
 		SequencerTime = SequencerTrailHierarchy->GetLocalTime();
 		TicksPerFrame = SequencerTrailHierarchy->GetFramesPerFrame();
+		StartFrame = SequencerTrailHierarchy->GetCurrentFramesInfo() ? SequencerTrailHierarchy->GetCurrentFramesInfo()->CurrentFrameTimes.StartFrame : SequencerTrailHierarchy->GetViewFrameRange().GetLowerBoundValue();
 	}
 }
 
diff --git a/Engine/Plugins/Editor/SequencerAnimTools/Source/SequencerAnimTools/Private/TrailHierarchy.cpp b/Engine/Plugins/Editor/SequencerAnimTools/Source/SequencerAnimTools/Private/TrailHierarchy.cpp
--- a/Engine/Plugins/Editor/SequencerAnimTools/Source/SequencerAnimTools/Private/TrailHierarchy.cpp
+++ b/Engine/Plugins/Editor/SequencerAnimTools/Source/SequencerAnimTools/Private/TrailHierarchy.cpp
@@ -52,7 +52,6 @@ void FTrail::GetColor(const FFrameNumber& CurrentTime, FColorState& State)
 		if (State.bFirstFrame == true)
 		{
 			State.bFirstFrame = false;
-			State.StartFrame = CurrentTime;
 		}
 
 		const int32 Index = (CurrentTime.Value - State.StartFrame.Value) / State.TicksPerFrame.Value;
