Anim In Engine: Make motion trails enter/exiting undo/redoable.
#trivial

diff --git a/Engine/Plugins/Editor/SequencerAnimTools/Source/SequencerAnimTools/Private/SequencerToolsEditMode.cpp b/Engine/Plugins/Editor/SequencerAnimTools/Source/SequencerAnimTools/Private/SequencerToolsEditMode.cpp
--- a/Engine/Plugins/Editor/SequencerAnimTools/Source/SequencerAnimTools/Private/SequencerToolsEditMode.cpp
+++ b/Engine/Plugins/Editor/SequencerAnimTools/Source/SequencerAnimTools/Private/SequencerToolsEditMode.cpp
@@ -54,7 +54,7 @@ void USequencerToolsEditMode::Enter()
 		{
 			LevelEditorPtr.Pin()->GetEditorModeManager().GetInteractiveToolsContext()->ToolManager->RegisterToolType(TEXT("SequencerMotionTrail"), NewObject<UMotionTrailToolBuilder>(this));
 			LevelEditorPtr.Pin()->GetEditorModeManager().GetInteractiveToolsContext()->ToolManager->RegisterToolType(TEXT("SequencerPivotTool"), NewObject<USequencerPivotToolBuilder>(this));
-
+			LevelEditorPtr.Pin()->GetEditorModeManager().GetInteractiveToolsContext()->ToolManager->ConfigureChangeTrackingMode(EToolChangeTrackingMode::FullUndoRedo);
 			// Currently there are some issues when the gizmo context object is shared among modes, so make sure to
 			// register it in the mode-level context store so it's not picked up by other modes.
 			UE::TransformGizmoUtil::RegisterTransformGizmoContextObject(GetInteractiveToolsContext(EToolsContextScope::EdMode));
