Pivot & MotionTrail tools: gizmo in world space

#rnx

diff --git a/Engine/Plugins/Editor/SequencerAnimTools/Source/SequencerAnimTools/Private/MotionTrailTool.cpp b/Engine/Plugins/Editor/SequencerAnimTools/Source/SequencerAnimTools/Private/MotionTrailTool.cpp
--- a/Engine/Plugins/Editor/SequencerAnimTools/Source/SequencerAnimTools/Private/MotionTrailTool.cpp
+++ b/Engine/Plugins/Editor/SequencerAnimTools/Source/SequencerAnimTools/Private/MotionTrailTool.cpp
@@ -403,9 +403,21 @@ void UMotionTrailTool::UpdateGizmoLocation()
 
 	if (TRSGizmo)
 	{
+		UInteractiveToolManager* ToolManager = GetToolManager();
+		const IToolsContextQueriesAPI* ToolsContextQueries = ToolManager ? ToolManager->GetContextQueriesAPI() : nullptr;
+		const bool bWorld = ToolsContextQueries ? ToolsContextQueries->GetCurrentCoordinateSystem() == EToolContextCoordinateSystem::World : false;
+
 		// note: this is a trick to avoid cycling notifications... (cf. UCombinedTransformGizmo::ReinitializeGizmoTransform)
 		TGuardValue<bool> PivotModeGuard(TransformProxy->bSetPivotMode, true);
-		TransformProxy->SetTransform(GizmoTransform);
+		if (bWorld)
+		{
+			const FTransform TransformNoRot(FQuat::Identity, GizmoTransform.GetTranslation(), GizmoTransform.GetScale3D());
+			TransformProxy->SetTransform(TransformNoRot);			
+		}
+		else
+		{
+			TransformProxy->SetTransform(GizmoTransform);
+		}
 	}
 }
 
diff --git a/Engine/Plugins/Editor/SequencerAnimTools/Source/SequencerAnimTools/Private/SequencerAnimEditPivotTool.cpp b/Engine/Plugins/Editor/SequencerAnimTools/Source/SequencerAnimTools/Private/SequencerAnimEditPivotTool.cpp
--- a/Engine/Plugins/Editor/SequencerAnimTools/Source/SequencerAnimTools/Private/SequencerAnimEditPivotTool.cpp
+++ b/Engine/Plugins/Editor/SequencerAnimTools/Source/SequencerAnimTools/Private/SequencerAnimEditPivotTool.cpp
@@ -976,10 +976,21 @@ void USequencerPivotTool::UpdateGizmoTransform()
 
 	if (TRSGizmo)
 	{
+		UInteractiveToolManager* ToolManager = GetToolManager();
+		const IToolsContextQueriesAPI* ToolsContextQueries = ToolManager ? ToolManager->GetContextQueriesAPI() : nullptr;
+		const bool bWorld = ToolsContextQueries ? ToolsContextQueries->GetCurrentCoordinateSystem() == EToolContextCoordinateSystem::World : false;
+
 		// note: this is a trick to avoid cycling notifications... (cf. UCombinedTransformGizmo::ReinitializeGizmoTransform)
 		TGuardValue<bool> PivotModeGuard(TransformProxy->bSetPivotMode, true);
-		//TransformProxy->SetTransform(GizmoTransformWithScale);
-		TransformProxy->SetTransform(GizmoTransform);
+		if (bWorld)
+		{
+			const FTransform TransformNoRot(FQuat::Identity, GizmoTransform.GetTranslation(), GizmoTransform.GetScale3D());
+			TransformProxy->SetTransform(TransformNoRot);			
+		}
+		else
+		{
+			TransformProxy->SetTransform(GizmoTransform);
+		}
 	}
 
 }
