Anim In Engine: Move commands registration for pivot tool to module startup.
#trivial

diff --git a/Engine/Plugins/Editor/SequencerAnimTools/Source/SequencerAnimTools/Private/SequencerAnimEditPivotTool.cpp b/Engine/Plugins/Editor/SequencerAnimTools/Source/SequencerAnimTools/Private/SequencerAnimEditPivotTool.cpp
--- a/Engine/Plugins/Editor/SequencerAnimTools/Source/SequencerAnimTools/Private/SequencerAnimEditPivotTool.cpp
+++ b/Engine/Plugins/Editor/SequencerAnimTools/Source/SequencerAnimTools/Private/SequencerAnimEditPivotTool.cpp
@@ -319,8 +319,6 @@ void USequencerPivotTool::Setup()
 
 	CommandBindings = MakeShareable(new FUICommandList);
 
-	FEditPivotCommands::Register();
-
 	const FEditPivotCommands& Commands = FEditPivotCommands::Get();
 
 	CommandBindings->MapAction(
diff --git a/Engine/Plugins/Editor/SequencerAnimTools/Source/SequencerAnimTools/Private/SequencerAnimToolsModule.cpp b/Engine/Plugins/Editor/SequencerAnimTools/Source/SequencerAnimTools/Private/SequencerAnimToolsModule.cpp
--- a/Engine/Plugins/Editor/SequencerAnimTools/Source/SequencerAnimTools/Private/SequencerAnimToolsModule.cpp
+++ b/Engine/Plugins/Editor/SequencerAnimTools/Source/SequencerAnimTools/Private/SequencerAnimToolsModule.cpp
@@ -1,6 +1,8 @@
 // Copyright Epic Games, Inc. All Rights Reserved.
 
 #include "SequencerAnimToolsModule.h"
+#include "SequencerAnimEditPivotTool.h"
+
 #include "EditorModeManager.h"
 #include "ILevelEditor.h"
 #include "LevelEditor.h"
@@ -26,6 +28,8 @@ void FSequencerAnimToolsModule::StartupModule()
 	{
 		MotionTrailOptions->OnDisplayPropertyChanged.AddRaw(this, &FSequencerAnimToolsModule::OnMotionTralOptionChanged);
 	}
+	FEditPivotCommands::Register();
+
 }
 
 void FSequencerAnimToolsModule::ShutdownModule()
@@ -37,6 +41,8 @@ void FSequencerAnimToolsModule::ShutdownModule()
 		MotionTrailOptions->OnDisplayPropertyChanged.RemoveAll(this);
 	}
 	*/
+	FEditPivotCommands::Unregister();
+
 }
 
 void FSequencerAnimToolsModule::OnMotionTralOptionChanged(FName PropertyName)
