Anim In Engine:  Register Motion Trail Commands early
#trivial

diff --git a/Engine/Plugins/Editor/SequencerAnimTools/Source/SequencerAnimTools/Private/MotionTrailTool.cpp b/Engine/Plugins/Editor/SequencerAnimTools/Source/SequencerAnimTools/Private/MotionTrailTool.cpp
--- a/Engine/Plugins/Editor/SequencerAnimTools/Source/SequencerAnimTools/Private/MotionTrailTool.cpp
+++ b/Engine/Plugins/Editor/SequencerAnimTools/Source/SequencerAnimTools/Private/MotionTrailTool.cpp
@@ -139,8 +139,6 @@ void UMotionTrailTool::SetupIntegration()
 
 	CommandBindings = MakeShareable(new FUICommandList);
 
-	FMotionTrailCommands::Register();
-
 	const FMotionTrailCommands& Commands = FMotionTrailCommands::Get();
 
 	CommandBindings->MapAction(
diff --git a/Engine/Plugins/Editor/SequencerAnimTools/Source/SequencerAnimTools/Private/SequencerAnimToolsModule.cpp b/Engine/Plugins/Editor/SequencerAnimTools/Source/SequencerAnimTools/Private/SequencerAnimToolsModule.cpp
--- a/Engine/Plugins/Editor/SequencerAnimTools/Source/SequencerAnimTools/Private/SequencerAnimToolsModule.cpp
+++ b/Engine/Plugins/Editor/SequencerAnimTools/Source/SequencerAnimTools/Private/SequencerAnimToolsModule.cpp
@@ -29,6 +29,7 @@ void FSequencerAnimToolsModule::StartupModule()
 		MotionTrailOptions->OnDisplayPropertyChanged.AddRaw(this, &FSequencerAnimToolsModule::OnMotionTralOptionChanged);
 	}
 	FEditPivotCommands::Register();
+	FMotionTrailCommands::Register();
 
 }
 
@@ -42,6 +43,7 @@ void FSequencerAnimToolsModule::ShutdownModule()
 	}
 	*/
 	FEditPivotCommands::Unregister();
+	FMotionTrailCommands::Unregister();
 
 }
 
