Motion trails: Added support for multiple rigs. Also fixed an issue with a recent change to SpawnableRestore where the Hierarchy may be empty and make sure anim layer tabs are closed by default.

diff --git a/Engine/Plugins/Animation/ControlRig/Source/ControlRigEditor/Private/EditMode/Settings/ControlRigUIRestoreStates.h b/Engine/Plugins/Animation/ControlRig/Source/ControlRigEditor/Private/EditMode/Settings/ControlRigUIRestoreStates.h
--- a/Engine/Plugins/Animation/ControlRig/Source/ControlRigEditor/Private/EditMode/Settings/ControlRigUIRestoreStates.h
+++ b/Engine/Plugins/Animation/ControlRig/Source/ControlRigEditor/Private/EditMode/Settings/ControlRigUIRestoreStates.h
@@ -17,7 +17,7 @@ struct FControlRigUIRestoreStates
 	bool bMotionTrailsOn = false;
 	
 	UPROPERTY()
-	bool bAnimLayerTabOpen = true;
+	bool bAnimLayerTabOpen = false;
 	
 	UPROPERTY()
 	bool bPoseTabOpen = false;
diff --git a/Engine/Plugins/Editor/SequencerAnimTools/Source/SequencerAnimTools/Private/Sequencer/SequencerTrailHierarchy.cpp b/Engine/Plugins/Editor/SequencerAnimTools/Source/SequencerAnimTools/Private/Sequencer/SequencerTrailHierarchy.cpp
--- a/Engine/Plugins/Editor/SequencerAnimTools/Source/SequencerAnimTools/Private/Sequencer/SequencerTrailHierarchy.cpp
+++ b/Engine/Plugins/Editor/SequencerAnimTools/Source/SequencerAnimTools/Private/Sequencer/SequencerTrailHierarchy.cpp
@@ -579,97 +579,164 @@ void FSequencerTrailHierarchy::OnBindingVisibilityStateChanged(UObject* BoundObj
 	}
 }
 
-void FSequencerTrailHierarchy::UpdateSequencerBindings(const TArray<FGuid>& SequencerBindings, TFunctionRef<void(UObject*, FTrail*, FGuid)> OnUpdated)
+namespace Private
 {
-	const FDateTime StartTime = FDateTime::Now();
 
-	TSharedPtr<ISequencer> Sequencer = WeakSequencer.Pin();
-	check(Sequencer);
+struct FEditorSelection
+{
+	FEditorSelection()
+	{
+		if (USelection* Selection = GEditor ? GEditor->GetSelectedActors() : nullptr)
+		{
+			Selection->GetSelectedObjects<AActor>(SelectedActors);
+			Selection->GetSelectedObjects<USceneComponent>(SelectedSceneComponents);
+		}
+	}
 
-	TArray<AActor*> SelectedActors;
-	GEditor->GetSelectedActors()->GetSelectedObjects<AActor>(SelectedActors);
-	TArray<USceneComponent*> SelectedSceneComponents;
-	GEditor->GetSelectedActors()->GetSelectedObjects<USceneComponent>(SelectedSceneComponents);
-	for (FGuid BindingGuid : SequencerBindings)
+	template<typename ComponentType>
+	TArray<ComponentType*> GetComponents(const TArrayView<TWeakObjectPtr<>>& InBoundObjects) const
 	{
-		bool bAddedControlRig = false;
-		if (UMovieSceneControlRigParameterTrack* CRParameterTrack = Sequencer->GetFocusedMovieSceneSequence()->GetMovieScene()->FindTrack<UMovieSceneControlRigParameterTrack>(BindingGuid))
+		TArray<ComponentType*> Components;
+		Components.Reserve(InBoundObjects.Num());
+
+		for (TWeakObjectPtr<> WeakBoundObject : InBoundObjects)
 		{
-			for (TWeakObjectPtr<> BoundObject : Sequencer->FindBoundObjects(BindingGuid, Sequencer->GetFocusedTemplateID()))
+			if (UObject* BoundObject = WeakBoundObject.Get())
 			{
-				if (!BoundObject.IsValid())
+				ComponentType* BoundComponent = Cast<ComponentType>(BoundObject);
+				if (AActor* BoundActor = Cast<AActor>(BoundObject))
+				{
+					if (SelectedActors.Contains(BoundActor) == false)
+					{
+						continue;
+					}
+					BoundComponent = Cast<ComponentType>(BoundActor->GetRootComponent());
+				}
+					
+				if (!BoundComponent || SelectedSceneComponents.Contains(BoundComponent) == false)
 				{
 					continue;
 				}
+					
+				Components.Add(BoundComponent);
+			}
+		}
+
+		return MoveTemp(Components);
+	}
 
-				USkeletalMeshComponent* BoundComponent = Cast<USkeletalMeshComponent>(BoundObject.Get());
-				if (AActor* BoundActor = Cast<AActor>(BoundObject.Get()))
+private:
+	TArray<AActor*> SelectedActors;
+	TArray<USceneComponent*> SelectedSceneComponents;
+};
+
+TArray<USkeletalMeshComponent*> GetBoundSkeletalMeshComponents(const TArrayView<TWeakObjectPtr<>>& InBoundObjects)
+{
+	TArray<USkeletalMeshComponent*> SkeletalMeshComponents;
+	SkeletalMeshComponents.Reserve(InBoundObjects.Num());
+
+	for (TWeakObjectPtr<> WeakBoundObject : InBoundObjects)
+	{
+		if (UObject* BoundObject = WeakBoundObject.Get())
+		{
+			USkeletalMeshComponent* BoundComponent = Cast<USkeletalMeshComponent>(BoundObject);
+			if (!BoundComponent)
+			{
+				if (AActor* BoundActor = Cast<AActor>(BoundObject))
 				{
 					BoundComponent = BoundActor->FindComponentByClass<USkeletalMeshComponent>();
 				}
+			}
 
-				if (!BoundComponent || !BoundComponent->GetSkeletalMeshAsset() || !BoundComponent->GetSkeletalMeshAsset()->GetSkeleton())
-				{
-					continue;
-				};
+			if (!BoundComponent || !BoundComponent->GetSkeletalMeshAsset() || !BoundComponent->GetSkeletalMeshAsset()->GetSkeleton())
+			{
+				continue;
+			}
 
-				//if we have a selected control rig don't add the transform track also, makes pinning messy
-				if (!ControlRigDelegateHandles.Contains(CRParameterTrack))
+			SkeletalMeshComponents.Add(BoundComponent);
+		}
+	}
+
+	return MoveTemp(SkeletalMeshComponents);
+}
+
+}
+	
+void FSequencerTrailHierarchy::UpdateSequencerBindings(const TArray<FGuid>& SequencerBindings, TFunctionRef<void(UObject*, FTrail*, FGuid)> OnUpdated)
+{
+	const FDateTime StartTime = FDateTime::Now();
+
+	TSharedPtr<ISequencer> Sequencer = WeakSequencer.IsValid() ? WeakSequencer.Pin() : nullptr;
+	if (!ensure(Sequencer))
+	{
+		return;
+	}
+
+	UMovieScene* MovieScene = Sequencer->GetFocusedMovieSceneSequence() ? Sequencer->GetFocusedMovieSceneSequence()->GetMovieScene() : nullptr;
+	if (!ensure(MovieScene))
+	{
+		return;
+	}
+	
+	const Private::FEditorSelection Selection;
+
+	for (const FGuid& BindingGuid : SequencerBindings)
+	{
+		const TArrayView<TWeakObjectPtr<>> BoundObjects = Sequencer->FindBoundObjects(BindingGuid, Sequencer->GetFocusedTemplateID());
+		if (BoundObjects.IsEmpty())
+		{
+			continue;
+		}
+
+		// try adding from control rig tracks
+		bool bAddedControlRig = false;
+		const TArray<USkeletalMeshComponent*>& SkeletalMeshComponents = Private::GetBoundSkeletalMeshComponents(BoundObjects);
+		const TArray<UMovieSceneTrack*> ControlRigTracks = MovieScene->FindTracks(UMovieSceneControlRigParameterTrack::StaticClass(), BindingGuid);
+		if (!ControlRigTracks.IsEmpty() && !SkeletalMeshComponents.IsEmpty())
+		{
+			for (UMovieSceneTrack* Track : ControlRigTracks)
+			{
+				if (UMovieSceneControlRigParameterTrack* CRParameterTrack = Cast<UMovieSceneControlRigParameterTrack>(Track))
 				{
-					RegisterControlRigDelegates(BoundComponent, CRParameterTrack);
-					UControlRig* ControlRig = CRParameterTrack->GetControlRig();
-					if (ControlRig)
+					for (USkeletalMeshComponent* BoundComponent: SkeletalMeshComponents)
 					{
-						TArray<FName> Selected = ControlRig->CurrentControlSelection();
-						for (const FName& ControlName : Selected)
+						//if we have a selected control rig don't add the transform track also, makes pinning messy
+						if (!ControlRigDelegateHandles.Contains(CRParameterTrack))
 						{
-							AddControlRigTrail(BoundComponent, ControlRig, CRParameterTrack, ControlName);
-							bAddedControlRig = true; 
+							RegisterControlRigDelegates(BoundComponent, CRParameterTrack);
+						
+							if (UControlRig* ControlRig = CRParameterTrack->GetControlRig())
+							{
+								const TArray<FName> Selected = ControlRig->CurrentControlSelection();
+								for (const FName& ControlName : Selected)
+								{
+									AddControlRigTrail(BoundComponent, ControlRig, CRParameterTrack, ControlName);
+									bAddedControlRig = true; 
+								}
+							}
+						}
+					
+						if (bAddedControlRig)
+						{
+							ClearSelection();
 						}
 					}
 				}
-				if (bAddedControlRig)
-				{
-					ClearSelection();
-					continue;
-				}
 			}
-		} // if ControlRigParameterTrack
+		}
+		// end try adding from control rig tracks
+		
 		if (bAddedControlRig)
 		{
 			continue;
-		}	
-		if (UMovieScene3DTransformTrack* TransformTrack = Sequencer->GetFocusedMovieSceneSequence()->GetMovieScene()->FindTrack<UMovieScene3DTransformTrack>(BindingGuid))
+		}
+
+		// try adding from transform track
+		UMovieScene3DTransformTrack* TransformTrack = MovieScene->FindTrack<UMovieScene3DTransformTrack>(BindingGuid);
+		if (TransformTrack)
 		{
-			for (TWeakObjectPtr<> BoundObject : Sequencer->FindBoundObjects(BindingGuid, Sequencer->GetFocusedTemplateID()))
+			for (USceneComponent* BoundComponent: Selection.GetComponents<USceneComponent>(BoundObjects))
 			{
-				if (!BoundObject.IsValid())
-				{
-					continue;
-				}
-			
-				/*
-				//if using old trails don't add new ones.
-				if (CVarUseOldSequencerMotionTrails->GetBool() == true)
-				{
-					continue;
-				}
-				*/
-
-				USceneComponent* BoundComponent = Cast<USceneComponent>(BoundObject.Get());
-				if (AActor* BoundActor = Cast<AActor>(BoundObject.Get()))
-				{
-					if (SelectedActors.Contains(BoundActor) == false)
-					{
-						continue;
-					}
-					BoundComponent = BoundActor->GetRootComponent();
-				}
-				else if (SelectedSceneComponents.Contains(BoundComponent) == false)
-				{
-					continue;
-				}
-
 				if (!ObjectsTracked.Contains(BoundComponent))
 				{
 					AddComponentToHierarchy(BindingGuid, BoundComponent, TransformTrack);
@@ -685,53 +752,31 @@ void FSequencerTrailHierarchy::UpdateSequencerBindings(const TArray<FGuid>& Sequ
 					OnUpdated(BoundComponent, AllTrails[ObjectsTracked[BoundComponent]].Get(), ObjectsTracked[BoundComponent]);
 				}
 			}
-		} // if TransformTrack
-		if (UMovieSceneSkeletalAnimationTrack* AnimTrack = Sequencer->GetFocusedMovieSceneSequence()->GetMovieScene()->FindTrack<UMovieSceneSkeletalAnimationTrack>(BindingGuid))
+		}
+		// end try adding from transform track
+
+		// try adding from anim track
+		if (UMovieSceneSkeletalAnimationTrack* AnimTrack = MovieScene->FindTrack<UMovieSceneSkeletalAnimationTrack>(BindingGuid))
 		{
-			for (TWeakObjectPtr<> BoundObject : Sequencer->FindBoundObjects(BindingGuid, Sequencer->GetFocusedTemplateID()))
+			for (USkeletalMeshComponent* BoundComponent: Selection.GetComponents<USkeletalMeshComponent>(BoundObjects))
 			{
-				if (!BoundObject.IsValid())
-				{
-					continue;
-				}
-				/*
-				if (CVarUseOldSequencerMotionTrails->GetBool() == true)
-				{
-					continue;
-				}
-				*/
-
-				USkeletalMeshComponent* BoundComponent = Cast<USkeletalMeshComponent>(BoundObject.Get());
-				if (AActor* BoundActor = Cast<AActor>(BoundObject.Get()))
-				{
-					if (SelectedActors.Contains(BoundActor) == false)
-					{
-						continue;
-					}
-					BoundComponent = BoundActor->FindComponentByClass<USkeletalMeshComponent>();
-				}
-				else if (SelectedSceneComponents.Contains(BoundComponent) == false)
-				{
-					continue;
-				}
-
 				if (!BoundComponent || !BoundComponent->GetSkeletalMeshAsset() || !BoundComponent->GetSkeletalMeshAsset()->GetSkeleton())
 				{
 					continue;
 				}
 
-				
-
 				if (!ObjectsTracked.Contains(BoundComponent))
 				{
-					if (UMovieScene3DTransformTrack* TransformTrack = Sequencer->GetFocusedMovieSceneSequence()->GetMovieScene()->FindTrack<UMovieScene3DTransformTrack>(BindingGuid))
+					if (TransformTrack)
 					{
-						AddComponentToHierarchy(BindingGuid,BoundComponent, TransformTrack);
+						AddComponentToHierarchy(BindingGuid, BoundComponent, TransformTrack);
 					}
 				}
 			}
 		}
+		// end try adding from anim track
 	}
+	
 	const FTimespan Timespan = FDateTime::Now() - StartTime;
 	TimingStats.Add("FSequencerTrailHierarchy::UpdateSequencerBindings", Timespan);
 }
@@ -965,12 +1010,11 @@ void FSequencerTrailHierarchy::RegisterControlRigDelegates(USkeletalMeshComponen
 		{
 		
 			TSharedPtr<ISequencer> Sequencer = WeakSequencer.Pin();
-			if (!Sequencer == false)
+			if (Sequencer.IsValid() == false)
 			{
 				return;
 			}
 			
-
 			const FRigControlElement* ControlElement = Cast<FRigControlElement>(InSubject.Element); 
 			if(ControlElement == nullptr)
 			{
diff --git a/Engine/Source/Editor/MovieSceneTools/Private/MovieSceneToolHelpers.cpp b/Engine/Source/Editor/MovieSceneTools/Private/MovieSceneToolHelpers.cpp
--- a/Engine/Source/Editor/MovieSceneTools/Private/MovieSceneToolHelpers.cpp
+++ b/Engine/Source/Editor/MovieSceneTools/Private/MovieSceneToolHelpers.cpp
@@ -4702,7 +4702,6 @@ void FAllSpawnableRestoreState::SaveStateForSpawnable(bool& bInOutWasChanged, TM
 		BindingLifetimeTrack->AddSection(*NewSection);
 	}
 };
-
 FAllSpawnableRestoreState::FAllSpawnableRestoreState(TSharedPtr<UE::MovieScene::FSharedPlaybackState> InSharedPlaybackState, UMovieScene* InMovieScene)
 	: SharedPlaybackState(InSharedPlaybackState)
 {
@@ -4725,24 +4724,40 @@ FAllSpawnableRestoreState::FAllSpawnableRestoreState(TSharedPtr<UE::MovieScene::
 		const FMovieSceneSequenceHierarchy* Hierarchy = SharedPlaybackStateToUse->GetHierarchy();
 		if (!Hierarchy)
 		{
-			UE_LOG(LogMovieScene, Warning, TEXT("FAllSpawnableRestoreState requires a valid FMovieSceneSequenceHierarchy from FSharedPlaybackState"));
-			return;
-		}
-		TArray<FMovieSceneSequenceID> SequenceIDs;
-		Hierarchy->AllSubSequenceIDs(SequenceIDs);
-		TArray<TWeakObjectPtr<UMovieSceneSequence>> Sequences;
-		for (const FMovieSceneSequenceID& SequenceID : SequenceIDs)
-		{
-			if (UMovieSceneSequence* Sequence = SharedPlaybackStateToUse->GetSequence(SequenceID))
+			if (UMovieSceneSequence* Sequence = SharedPlaybackStateToUse->GetSequence(MovieSceneSequenceID::Root))
 			{
 				if (UMovieScene* MovieScene = Sequence->GetMovieScene())
 				{
 					FPerMovieScene PerMovieScene;
 					PerMovieScene.WeakMovieScene = MovieScene;
-					PerMovieScene.SequenceID = SequenceID;
+					PerMovieScene.SequenceID = MovieSceneSequenceID::Root;
 					OwnershipMaps.Add(MovieScene, PerMovieScene);
 				}
 			}
+			else
+			{
+				UE_LOG(LogMovieScene, Warning, TEXT("FAllSpawnableRestoreState requires a valid Root Sequence from FSharedPlaybackState"));
+				return;
+			}
+		}
+		else
+		{
+			TArray<FMovieSceneSequenceID> SequenceIDs;
+			Hierarchy->AllSubSequenceIDs(SequenceIDs);
+			TArray<TWeakObjectPtr<UMovieSceneSequence>> Sequences;
+			for (const FMovieSceneSequenceID& SequenceID : SequenceIDs)
+			{
+				if (UMovieSceneSequence* Sequence = SharedPlaybackStateToUse->GetSequence(SequenceID))
+				{
+					if (UMovieScene* MovieScene = Sequence->GetMovieScene())
+					{
+						FPerMovieScene PerMovieScene;
+						PerMovieScene.WeakMovieScene = MovieScene;
+						PerMovieScene.SequenceID = SequenceID;
+						OwnershipMaps.Add(MovieScene, PerMovieScene);
+					}
+				}
+			}
 		}
 	}
 
