Anim In Engine: Fix actor motion trails not working

diff --git a/Engine/Plugins/Editor/SequencerAnimTools/Source/SequencerAnimTools/Private/Sequencer/SequencerTrailHierarchy.cpp b/Engine/Plugins/Editor/SequencerAnimTools/Source/SequencerAnimTools/Private/Sequencer/SequencerTrailHierarchy.cpp
--- a/Engine/Plugins/Editor/SequencerAnimTools/Source/SequencerAnimTools/Private/Sequencer/SequencerTrailHierarchy.cpp
+++ b/Engine/Plugins/Editor/SequencerAnimTools/Source/SequencerAnimTools/Private/Sequencer/SequencerTrailHierarchy.cpp
@@ -316,9 +316,7 @@ void FSequencerTrailHierarchy::CalculateEvalRangeArray()
 
 FFrameNumber FSequencerTrailHierarchy::GetFramesPerSegment() const
 {
-	FFrameNumber FramesPerTick = FFrameRate::TransformTime(FFrameNumber(1), WeakSequencer.Pin()->GetFocusedDisplayRate(), 
-		WeakSequencer.Pin()->GetFocusedTickResolution()).RoundToFrame();
-	return FramesPerTick;
+	return GetFramesPerFrame();
 }
 
 const FCurrentFramesInfo* FSequencerTrailHierarchy::GetCurrentFramesInfo() const
@@ -328,9 +326,12 @@ const FCurrentFramesInfo* FSequencerTrailHierarchy::GetCurrentFramesInfo() const
 
 FFrameNumber FSequencerTrailHierarchy::GetFramesPerFrame() const
 {
-	FFrameNumber FramesPerTick = FFrameRate::TransformTime(FFrameNumber(1), WeakSequencer.Pin()->GetFocusedDisplayRate(),
-		WeakSequencer.Pin()->GetFocusedTickResolution()).RoundToFrame();
-
+	FFrameNumber FramesPerTick(0);
+	if (TSharedPtr<ISequencer> Sequencer = WeakSequencer.Pin())
+	{
+		FramesPerTick = FFrameRate::TransformTime(FFrameNumber(1), Sequencer->GetFocusedDisplayRate(),
+			Sequencer->GetFocusedTickResolution()).RoundToFrame();
+	}
 	return FramesPerTick;
 }
 
@@ -579,93 +580,9 @@ void FSequencerTrailHierarchy::OnBindingVisibilityStateChanged(UObject* BoundObj
 	}
 }
 
-namespace Private
-{
-
-struct FEditorSelection
-{
-	FEditorSelection()
-	{
-		if (USelection* Selection = GEditor ? GEditor->GetSelectedActors() : nullptr)
-		{
-			Selection->GetSelectedObjects<AActor>(SelectedActors);
-			Selection->GetSelectedObjects<USceneComponent>(SelectedSceneComponents);
-		}
-	}
-
-	template<typename ComponentType>
-	TArray<ComponentType*> GetComponents(const TArrayView<TWeakObjectPtr<>>& InBoundObjects) const
-	{
-		TArray<ComponentType*> Components;
-		Components.Reserve(InBoundObjects.Num());
-
-		for (TWeakObjectPtr<> WeakBoundObject : InBoundObjects)
-		{
-			if (UObject* BoundObject = WeakBoundObject.Get())
-			{
-				ComponentType* BoundComponent = Cast<ComponentType>(BoundObject);
-				if (AActor* BoundActor = Cast<AActor>(BoundObject))
-				{
-					if (SelectedActors.Contains(BoundActor) == false)
-					{
-						continue;
-					}
-					BoundComponent = Cast<ComponentType>(BoundActor->GetRootComponent());
-				}
-					
-				if (!BoundComponent || SelectedSceneComponents.Contains(BoundComponent) == false)
-				{
-					continue;
-				}
-					
-				Components.Add(BoundComponent);
-			}
-		}
-
-		return MoveTemp(Components);
-	}
-
-private:
-	TArray<AActor*> SelectedActors;
-	TArray<USceneComponent*> SelectedSceneComponents;
-};
-
-TArray<USkeletalMeshComponent*> GetBoundSkeletalMeshComponents(const TArrayView<TWeakObjectPtr<>>& InBoundObjects)
-{
-	TArray<USkeletalMeshComponent*> SkeletalMeshComponents;
-	SkeletalMeshComponents.Reserve(InBoundObjects.Num());
-
-	for (TWeakObjectPtr<> WeakBoundObject : InBoundObjects)
-	{
-		if (UObject* BoundObject = WeakBoundObject.Get())
-		{
-			USkeletalMeshComponent* BoundComponent = Cast<USkeletalMeshComponent>(BoundObject);
-			if (!BoundComponent)
-			{
-				if (AActor* BoundActor = Cast<AActor>(BoundObject))
-				{
-					BoundComponent = BoundActor->FindComponentByClass<USkeletalMeshComponent>();
-				}
-			}
-
-			if (!BoundComponent || !BoundComponent->GetSkeletalMeshAsset() || !BoundComponent->GetSkeletalMeshAsset()->GetSkeleton())
-			{
-				continue;
-			}
-
-			SkeletalMeshComponents.Add(BoundComponent);
-		}
-	}
-
-	return MoveTemp(SkeletalMeshComponents);
-}
-
-}
-	
 void FSequencerTrailHierarchy::UpdateSequencerBindings(const TArray<FGuid>& SequencerBindings, TFunctionRef<void(UObject*, FTrail*, FGuid)> OnUpdated)
 {
 	const FDateTime StartTime = FDateTime::Now();
-
 	TSharedPtr<ISequencer> Sequencer = WeakSequencer.IsValid() ? WeakSequencer.Pin() : nullptr;
 	if (!ensure(Sequencer))
 	{
@@ -677,66 +594,92 @@ void FSequencerTrailHierarchy::UpdateSequencerBindings(const TArray<FGuid>& Sequ
 	{
 		return;
 	}
-	
-	const Private::FEditorSelection Selection;
 
-	for (const FGuid& BindingGuid : SequencerBindings)
-	{
-		const TArrayView<TWeakObjectPtr<>> BoundObjects = Sequencer->FindBoundObjects(BindingGuid, Sequencer->GetFocusedTemplateID());
-		if (BoundObjects.IsEmpty())
-		{
-			continue;
-		}
+	TArray<AActor*> SelectedActors;
+	GEditor->GetSelectedActors()->GetSelectedObjects<AActor>(SelectedActors);
+	TArray<USceneComponent*> SelectedSceneComponents;
+	GEditor->GetSelectedActors()->GetSelectedObjects<USceneComponent>(SelectedSceneComponents);
 
-		// try adding from control rig tracks
+	for (FGuid BindingGuid : SequencerBindings)
+	{
 		bool bAddedControlRig = false;
-		const TArray<USkeletalMeshComponent*>& SkeletalMeshComponents = Private::GetBoundSkeletalMeshComponents(BoundObjects);
 		const TArray<UMovieSceneTrack*> ControlRigTracks = MovieScene->FindTracks(UMovieSceneControlRigParameterTrack::StaticClass(), BindingGuid);
-		if (!ControlRigTracks.IsEmpty() && !SkeletalMeshComponents.IsEmpty())
+		for (UMovieSceneTrack* Track : ControlRigTracks)
 		{
-			for (UMovieSceneTrack* Track : ControlRigTracks)
+			if (UMovieSceneControlRigParameterTrack* CRParameterTrack = Cast<UMovieSceneControlRigParameterTrack>(Track))
 			{
-				if (UMovieSceneControlRigParameterTrack* CRParameterTrack = Cast<UMovieSceneControlRigParameterTrack>(Track))
+				for (TWeakObjectPtr<> BoundObject : Sequencer->FindBoundObjects(BindingGuid, Sequencer->GetFocusedTemplateID()))
 				{
-					for (USkeletalMeshComponent* BoundComponent: SkeletalMeshComponents)
+					if (!BoundObject.IsValid())
+					{
+						continue;
+					}
+
+					USkeletalMeshComponent* BoundComponent = Cast<USkeletalMeshComponent>(BoundObject.Get());
+					if (AActor* BoundActor = Cast<AActor>(BoundObject.Get()))
+					{
+						BoundComponent = BoundActor->FindComponentByClass<USkeletalMeshComponent>();
+					}
+
+					if (!BoundComponent || !BoundComponent->GetSkeletalMeshAsset() || !BoundComponent->GetSkeletalMeshAsset()->GetSkeleton())
+					{
+						continue;
+					};
+
+					//if we have a selected control rig don't add the transform track also, makes pinning messy
+					if (!ControlRigDelegateHandles.Contains(CRParameterTrack))
 					{
-						//if we have a selected control rig don't add the transform track also, makes pinning messy
-						if (!ControlRigDelegateHandles.Contains(CRParameterTrack))
+						RegisterControlRigDelegates(BoundComponent, CRParameterTrack);
+						UControlRig* ControlRig = CRParameterTrack->GetControlRig();
+						if (ControlRig)
 						{
-							RegisterControlRigDelegates(BoundComponent, CRParameterTrack);
-						
-							if (UControlRig* ControlRig = CRParameterTrack->GetControlRig())
+							TArray<FName> Selected = ControlRig->CurrentControlSelection();
+							for (const FName& ControlName : Selected)
 							{
-								const TArray<FName> Selected = ControlRig->CurrentControlSelection();
-								for (const FName& ControlName : Selected)
-								{
-									AddControlRigTrail(BoundComponent, ControlRig, CRParameterTrack, ControlName);
-									bAddedControlRig = true; 
-								}
+								AddControlRigTrail(BoundComponent, ControlRig, CRParameterTrack, ControlName);
+								bAddedControlRig = true;
 							}
 						}
-					
-						if (bAddedControlRig)
-						{
-							ClearSelection();
-						}
 					}
 				}
 			}
-		}
-		// end try adding from control rig tracks
-		
+		} // if ControlRigParameterTrack
 		if (bAddedControlRig)
 		{
+			ClearSelection();
 			continue;
 		}
-
-		// try adding from transform track
-		UMovieScene3DTransformTrack* TransformTrack = MovieScene->FindTrack<UMovieScene3DTransformTrack>(BindingGuid);
-		if (TransformTrack)
+		if (UMovieScene3DTransformTrack* TransformTrack = Sequencer->GetFocusedMovieSceneSequence()->GetMovieScene()->FindTrack<UMovieScene3DTransformTrack>(BindingGuid))
 		{
-			for (USceneComponent* BoundComponent: Selection.GetComponents<USceneComponent>(BoundObjects))
+			for (TWeakObjectPtr<> BoundObject : Sequencer->FindBoundObjects(BindingGuid, Sequencer->GetFocusedTemplateID()))
 			{
+				if (!BoundObject.IsValid())
+				{
+					continue;
+				}
+
+				/*
+				//if using old trails don't add new ones.
+				if (CVarUseOldSequencerMotionTrails->GetBool() == true)
+				{
+					continue;
+				}
+				*/
+
+				USceneComponent* BoundComponent = Cast<USceneComponent>(BoundObject.Get());
+				if (AActor* BoundActor = Cast<AActor>(BoundObject.Get()))
+				{
+					if (SelectedActors.Contains(BoundActor) == false)
+					{
+						continue;
+					}
+					BoundComponent = BoundActor->GetRootComponent();
+				}
+				else if (SelectedSceneComponents.Contains(BoundComponent) == false)
+				{
+					continue;
+				}
+
 				if (!ObjectsTracked.Contains(BoundComponent))
 				{
 					AddComponentToHierarchy(BindingGuid, BoundComponent, TransformTrack);
@@ -752,31 +695,53 @@ void FSequencerTrailHierarchy::UpdateSequencerBindings(const TArray<FGuid>& Sequ
 					OnUpdated(BoundComponent, AllTrails[ObjectsTracked[BoundComponent]].Get(), ObjectsTracked[BoundComponent]);
 				}
 			}
-		}
-		// end try adding from transform track
-
-		// try adding from anim track
-		if (UMovieSceneSkeletalAnimationTrack* AnimTrack = MovieScene->FindTrack<UMovieSceneSkeletalAnimationTrack>(BindingGuid))
+		} // if TransformTrack
+		if (UMovieSceneSkeletalAnimationTrack* AnimTrack = Sequencer->GetFocusedMovieSceneSequence()->GetMovieScene()->FindTrack<UMovieSceneSkeletalAnimationTrack>(BindingGuid))
 		{
-			for (USkeletalMeshComponent* BoundComponent: Selection.GetComponents<USkeletalMeshComponent>(BoundObjects))
+			for (TWeakObjectPtr<> BoundObject : Sequencer->FindBoundObjects(BindingGuid, Sequencer->GetFocusedTemplateID()))
 			{
+				if (!BoundObject.IsValid())
+				{
+					continue;
+				}
+				/*
+				if (CVarUseOldSequencerMotionTrails->GetBool() == true)
+				{
+					continue;
+				}
+				*/
+
+				USkeletalMeshComponent* BoundComponent = Cast<USkeletalMeshComponent>(BoundObject.Get());
+				if (AActor* BoundActor = Cast<AActor>(BoundObject.Get()))
+				{
+					if (SelectedActors.Contains(BoundActor) == false)
+					{
+						continue;
+					}
+					BoundComponent = BoundActor->FindComponentByClass<USkeletalMeshComponent>();
+				}
+				else if (SelectedSceneComponents.Contains(BoundComponent) == false)
+				{
+					continue;
+				}
+
 				if (!BoundComponent || !BoundComponent->GetSkeletalMeshAsset() || !BoundComponent->GetSkeletalMeshAsset()->GetSkeleton())
 				{
 					continue;
 				}
 
+
+
 				if (!ObjectsTracked.Contains(BoundComponent))
 				{
-					if (TransformTrack)
+					if (UMovieScene3DTransformTrack* TransformTrack = Sequencer->GetFocusedMovieSceneSequence()->GetMovieScene()->FindTrack<UMovieScene3DTransformTrack>(BindingGuid))
 					{
 						AddComponentToHierarchy(BindingGuid, BoundComponent, TransformTrack);
 					}
 				}
 			}
 		}
-		// end try adding from anim track
 	}
-	
 	const FTimespan Timespan = FDateTime::Now() - StartTime;
 	TimingStats.Add("FSequencerTrailHierarchy::UpdateSequencerBindings", Timespan);
 }
@@ -1266,8 +1231,10 @@ void FSequencerTrailHierarchy::OnSetLinearColor(FGuid InGuid, FLinearColor Color
 {
 	if (const TUniquePtr<FTrail>* Trail = GetAllTrails().Find(InGuid))
 	{
-		FTrajectoryDrawInfo* CurDrawInfo = (*Trail)->GetDrawInfo();
-		CurDrawInfo->SetColor(Color);
+		if (FTrajectoryDrawInfo* CurDrawInfo = (*Trail)->GetDrawInfo())
+		{
+			CurDrawInfo->SetColor(Color);
+		}
 	}
 }
 void FSequencerTrailHierarchy::OnSetHasOffset(FGuid InGuid, bool bOffset)
