Anim In Engine: Pivot Tool fixes. Works with rotation orders, only select controls you want to transform when using edit pivot tool, make sure to not evaluate sequencer when using tool and some cleanup to use the same getters.

diff --git a/Engine/Plugins/Animation/ControlRig/Source/ControlRigEditor/Private/EditMode/ControlRigEditMode.cpp b/Engine/Plugins/Animation/ControlRig/Source/ControlRigEditor/Private/EditMode/ControlRigEditMode.cpp
--- a/Engine/Plugins/Animation/ControlRig/Source/ControlRigEditor/Private/EditMode/ControlRigEditMode.cpp
+++ b/Engine/Plugins/Animation/ControlRig/Source/ControlRigEditor/Private/EditMode/ControlRigEditMode.cpp
@@ -3063,7 +3063,7 @@ void FControlRigEditMode::UpdatePivotTransformsIfNeeded(UControlRig* InControlRi
 		return;
 	}
 
-	if (!InControlRig)
+	if (!InControlRig || !InControlRig->GetHierarchy())
 	{
 		return;
 	}
@@ -3078,8 +3078,11 @@ void FControlRigEditMode::UpdatePivotTransformsIfNeeded(UControlRig* InControlRi
 		}
 		for (AControlRigShapeActor* ShapeActor : *ShapeActors)
 		{
-			const FTransform Transform = InControlRig->GetControlGlobalTransform(ShapeActor->ControlName);
-			ShapeActor->SetActorTransform(Transform * ComponentTransform);
+			if (FRigTransformElement* ControlElement = InControlRig->FindControl(ShapeActor->ControlName))
+			{
+				const FTransform Transform = InControlRig->GetHierarchy()->GetTransform(ControlElement, ERigTransformType::CurrentGlobal);
+				ShapeActor->SetActorTransform(Transform * ComponentTransform);
+			}
 		}
 	}
 
diff --git a/Engine/Plugins/Animation/ControlRig/Source/ControlRigEditor/Private/Tools/ControlRigSnapper.cpp b/Engine/Plugins/Animation/ControlRig/Source/ControlRigEditor/Private/Tools/ControlRigSnapper.cpp
--- a/Engine/Plugins/Animation/ControlRig/Source/ControlRigEditor/Private/Tools/ControlRigSnapper.cpp
+++ b/Engine/Plugins/Animation/ControlRig/Source/ControlRigEditor/Private/Tools/ControlRigSnapper.cpp
@@ -118,12 +118,18 @@ static bool LocalGetControlRigControlTransforms(IMovieScenePlayer* Player, const
 	UControlRig* ControlRig, const FName& ControlName,
 	const TArray<FFrameNumber>& Frames, const TArray<FTransform>& ParentTransforms, TArray<FTransform>& OutTransforms)
 {
+	if (ControlRig == nullptr || ControlRig->GetHierarchy() == nullptr)
+	{
+		UE_LOG(LogControlRig, Error, TEXT("No Control Rig or Hierarchy"));
+		return false;
+	}
 	if (Frames.Num() > ParentTransforms.Num())
 	{
 		UE_LOG(LogControlRig, Error, TEXT("Number of Frames %d to Snap greater  than Parent Frames %d"), Frames.Num(), ParentTransforms.Num());
 		return false;
 	}
-	if (ControlRig->FindControl(ControlName) == nullptr)
+	FRigTransformElement* ControlElement = ControlRig->FindControl(ControlName);
+	if (ControlElement == nullptr)
 	{
 		UE_LOG(LogControlRig, Error, TEXT("Can not find Control %s"), *(ControlName.ToString()));
 		return false;
@@ -162,7 +168,7 @@ static bool LocalGetControlRigControlTransforms(IMovieScenePlayer* Player, const
 			{
 				ControlRig->Evaluate_AnyThread();
 			}
-			OutTransforms[Index] = ControlRig->GetControlGlobalTransform(ControlName) * ParentTransforms[Index];
+			OutTransforms[Index] = ControlRig->GetHierarchy()->GetTransform(ControlElement, ERigTransformType::CurrentGlobal) * ParentTransforms[Index];
 		}
 	}
 	return true;
diff --git a/Engine/Plugins/Animation/ControlRig/Source/ControlRigEditor/Private/Tools/EvaluateSequencerTools.cpp b/Engine/Plugins/Animation/ControlRig/Source/ControlRigEditor/Private/Tools/EvaluateSequencerTools.cpp
--- a/Engine/Plugins/Animation/ControlRig/Source/ControlRigEditor/Private/Tools/EvaluateSequencerTools.cpp
+++ b/Engine/Plugins/Animation/ControlRig/Source/ControlRigEditor/Private/Tools/EvaluateSequencerTools.cpp
@@ -386,58 +386,20 @@ namespace AIE
 	}
 
 
-	void FSetTransformHelpers::SetControlTransform(const FEulerTransform& EulerTransform, UControlRig* ControlRig, FRigControlElement* ControlElement, const FRigControlModifiedContext& Context)
+	void FSetTransformHelpers::SetControlTransform(const FEulerTransform& InLocalEulerTransform, UControlRig* ControlRig, FRigControlElement* ControlElement, const FRigControlModifiedContext& Context)
 	{
 		if (ControlElement && ControlRig)
 		{
-			FVector TLocation = EulerTransform.GetLocation();
-			FRotator TRotation = EulerTransform.Rotation;
-			FVector TScale = EulerTransform.Scale;
-			FTransform RealTransform(TRotation, TLocation, TScale);
-			if (SetConstrainedTransform(RealTransform, ControlRig, ControlElement, Context))
+			FVector TLocation = InLocalEulerTransform.GetLocation();
+			FRotator TRotation = InLocalEulerTransform.Rotation;
+			FVector TScale = InLocalEulerTransform.Scale;
+			FTransform LocalTransform(TRotation, TLocation, TScale);
+			if (SetConstrainedTransform(LocalTransform, ControlRig, ControlElement, Context))
 			{
 				return;
 			}
-			switch (ControlElement->Settings.ControlType)
-			{
-			case ERigControlType::Transform:
-			{
-				FVector EulerAngle(TRotation.Roll, TRotation.Pitch, TRotation.Yaw);
-				ControlRig->GetHierarchy()->SetControlSpecifiedEulerAngle(ControlElement, EulerAngle);
-				ControlRig->SetControlValue<FRigControlValue::FTransform_Float>(ControlElement->GetKey().Name, RealTransform, true, Context, false);
-				ControlRig->GetHierarchy()->SetControlSpecifiedEulerAngle(ControlElement, EulerAngle);
-				break;
-
-			}
-			case ERigControlType::TransformNoScale:
-			{
-				FTransformNoScale NoScale(TLocation, TRotation.Quaternion());
-				ControlRig->SetControlValue<FRigControlValue::FTransformNoScale_Float>(ControlElement->GetKey().Name, NoScale, true, Context, false);
-				break;
-
-			}
-			case ERigControlType::EulerTransform:
-			{
-				URigHierarchy* Hierarchy = ControlRig->GetHierarchy();
 
-				if (Hierarchy->UsesPreferredEulerAngles())
-				{
-					FVector EulerAngle(TRotation.Roll, TRotation.Pitch, TRotation.Yaw);
-					FQuat Quat = Hierarchy->GetControlQuaternion(ControlElement, EulerAngle);
-					Hierarchy->SetControlSpecifiedEulerAngle(ControlElement, EulerAngle);
-					FRotator UERotator(Quat);
-					FEulerTransform UETransform(UERotator, TLocation, TScale);
-					UETransform.Rotation = UERotator;
-					ControlRig->SetControlValue<FRigControlValue::FEulerTransform_Float>(ControlElement->GetKey().Name, UETransform, true, Context, false);
-					Hierarchy->SetControlSpecifiedEulerAngle(ControlElement, EulerAngle);
-				}
-				else
-				{
-					ControlRig->SetControlValue<FRigControlValue::FEulerTransform_Float>(ControlElement->GetKey().Name, FEulerTransform(RealTransform), true, Context, false);
-				}
-				break;
-			}
-			}
+			ControlRig->SetControlLocalTransform(ControlElement->GetKey().Name, LocalTransform, true, Context, false /*undo*/, true/* bFixEulerFlips*/);
 			ControlRig->Evaluate_AnyThread();
 		}
 	}
diff --git a/Engine/Plugins/Editor/SequencerAnimTools/Source/SequencerAnimTools/Private/SequencerAnimEditPivotTool.cpp b/Engine/Plugins/Editor/SequencerAnimTools/Source/SequencerAnimTools/Private/SequencerAnimEditPivotTool.cpp
--- a/Engine/Plugins/Editor/SequencerAnimTools/Source/SequencerAnimTools/Private/SequencerAnimEditPivotTool.cpp
+++ b/Engine/Plugins/Editor/SequencerAnimTools/Source/SequencerAnimTools/Private/SequencerAnimEditPivotTool.cpp
@@ -39,6 +39,7 @@
 #include "SLevelViewport.h"
 #include "EditorGizmos/EditorTransformGizmoUtil.h"
 #include "EditorGizmos/TransformGizmo.h"
+#include "EditMode/ControlRigEditMode.h"
 
 #define LOCTEXT_NAMESPACE "SequencerAnimTools"
 
@@ -52,6 +53,62 @@ void FEditPivotCommands::RegisterCommands()
 FSavedMappings USequencerPivotTool::SavedPivotLocations;
 FLastSelectedObjects USequencerPivotTool::LastSelectedObjects;
 
+namespace UE::AIE::PivotTool
+{
+
+	TArray<FName> CurrentControlSelection(UControlRig* InControlRig) 
+	{
+		TArray<FName> Names;
+		if (InControlRig)
+		{
+			// use this version after 5.7.1 is done, expose FControlRigEditMode::GetRigElementsForSettingTransforms
+			//TArray<FRigElementKey> SelectedRigElements = FControlRigEditMode::GetRigElementsForSettingTransforms(InControlRig, true /*bSelectionOnly*/, false /*bIncludeChannels*/);
+			/*for (const FRigElementKey& Key : SelectedRigElements)
+			{
+				Names.Add(Key.Name);
+			}
+			*/
+
+			TArray<FName> SelectedControls = InControlRig->CurrentControlSelection();
+			for (const FName& Name : SelectedControls)
+			{
+				if (FRigControlElement* Element = InControlRig->FindControl(Name))
+				{
+					if (Element->IsAnimationChannel() == false)
+					{
+						Names.Add(Name);
+					}
+				}
+			}
+		}
+
+		return Names;
+	}
+	//includes global + hosting transform
+	FTransform GetControlWorldSpaceTransform(UControlRig* ControlRig, const FName& Control)
+	{
+		FTransform WorldTransform = FTransform::Identity;
+		if (ControlRig && ControlRig->GetHierarchy())
+		{
+			if (FRigTransformElement* ControlElement = ControlRig->FindControl(Control))
+			{
+				TSharedPtr<IControlRigObjectBinding> ObjectBinding = ControlRig->GetObjectBinding();
+				FTransform ParentTransform = FTransform::Identity;
+				if (ObjectBinding.IsValid())
+				{
+					if (USceneComponent* HostingComponent = Cast<USceneComponent>(ObjectBinding->GetBoundObject()))
+					{
+						ParentTransform = HostingComponent->GetComponentTransform();
+					}
+				}
+				//Use GetTransform() not GetControlGlobalTransform
+				WorldTransform = ControlRig->GetHierarchy()->GetTransform(ControlElement, ERigTransformType::CurrentGlobal);
+				WorldTransform = WorldTransform * ParentTransform;
+			}
+		}
+		return WorldTransform;
+	}
+}
 
 static void GetControlRigsAndSequencer(TArray<TWeakObjectPtr<UControlRig>>& ControlRigs, TWeakPtr<ISequencer>& SequencerPtr, ULevelSequence** LevelSequence)
 {
@@ -65,8 +122,13 @@ static void GetControlRigsAndSequencer(TArray<TWeakObjectPtr<UControlRig>>& Cont
 		SequencerPtr = LevelSequenceEditor ? LevelSequenceEditor->GetSequencer() : nullptr;
 		if (SequencerPtr.IsValid())
 		{
+			//use editt mode (which must be active for pivot tool) to get 
 			TArray<UControlRig*> TempControlRigs;
-			TempControlRigs = UControlRigSequencerEditorLibrary::GetVisibleControlRigs();
+			FControlRigEditMode* ControlRigEditMode = static_cast<FControlRigEditMode*>(GLevelEditorModeTools().GetActiveMode(FControlRigEditMode::ModeName));
+			if (ControlRigEditMode)
+			{
+				TempControlRigs = ControlRigEditMode->GetControlRigsArray(false /*bIsVisible*/); //get all may want to pivot non-visible
+			}
 			for (UControlRig* ControlRig : TempControlRigs)
 			{
 				UControlRig* Base = ControlRig;
@@ -133,15 +195,18 @@ TOptional<FTransform> FControlRigMappings::GetWorldTransform(const FName& Name)
 {
 	
 	TOptional<FTransform> WorldTransform;
-	if (ControlRig.IsValid())
+	if (ControlRig.IsValid() && ControlRig->GetHierarchy())
 	{
-		const FTransform* LocalTransform = PivotTransforms.Find(Name);
-		if (LocalTransform)
+		if (FRigTransformElement* ControlElement = ControlRig->FindControl(Name))
 		{
-			FTransform ToWorldTransform = GetParentTransform();
-			FTransform CurrentTransform = ControlRig.Get()->GetControlGlobalTransform(Name);
-			WorldTransform = (*LocalTransform) * CurrentTransform * ToWorldTransform;
-
+			const FTransform* LocalTransform = PivotTransforms.Find(Name);
+			if (LocalTransform)
+			{
+				FTransform ToWorldTransform = GetParentTransform();
+				//Use GetTransform() not GetControlGlobalTransform
+				FTransform CurrentTransform = ControlRig->GetHierarchy()->GetTransform(ControlElement, ERigTransformType::CurrentGlobal);
+				WorldTransform = (*LocalTransform) * CurrentTransform * ToWorldTransform;
+			}
 		}
 	}
 	return WorldTransform;
@@ -149,12 +214,16 @@ TOptional<FTransform> FControlRigMappings::GetWorldTransform(const FName& Name)
 
 void FControlRigMappings::SetFromWorldTransform(const FName& Name, const FTransform& WorldTransform)
 {
-	if (ControlRig.IsValid())
+	if (ControlRig.IsValid() && ControlRig->GetHierarchy())
 	{
-		FTransform ToWorldTransform = GetParentTransform();
-		FTransform GlobalTransform = ControlRig.Get()->GetControlGlobalTransform(Name) * ToWorldTransform;
-		FTransform LocalTransform = WorldTransform.GetRelativeTransform(GlobalTransform);
-		PivotTransforms.Add(Name, LocalTransform);
+		if (FRigTransformElement* ControlElement = ControlRig->FindControl(Name))
+		{
+			FTransform ToWorldTransform = GetParentTransform();
+			FTransform GlobalTransform = ControlRig->GetHierarchy()->GetTransform(ControlElement, ERigTransformType::CurrentGlobal) * ToWorldTransform;
+
+			FTransform LocalTransform = WorldTransform.GetRelativeTransform(GlobalTransform);
+			PivotTransforms.Add(Name, LocalTransform);
+		}
 	}
 }
 
@@ -371,7 +440,7 @@ void USequencerPivotTool::SaveLastSelected()
 	{
 		if (ControlRig.IsValid())
 		{
-			TArray<FName> SelectedControls = ControlRig->CurrentControlSelection();
+			TArray<FName> SelectedControls = UE::AIE::PivotTool::CurrentControlSelection(ControlRig.Get());
 			for (const FName& Name : SelectedControls)
 			{
 				FControlRigMappings Mapping;
@@ -448,7 +517,7 @@ bool USequencerPivotTool::SetGizmoBasedOnSelection(bool bUseSaved)
 			continue;
 		}
 		FControlRigMappings* Mappings = SavedPivotLocations.ControlRigMappings.Find(ControlRig);
-		TArray<FName> SelectedControls = ControlRig->CurrentControlSelection();
+		TArray<FName> SelectedControls = UE::AIE::PivotTool::CurrentControlSelection(ControlRig.Get());
 
 		for (const FName& Name : SelectedControls)
 		{
@@ -466,8 +535,8 @@ bool USequencerPivotTool::SetGizmoBasedOnSelection(bool bUseSaved)
 			if (bHasSavedPivot == false)
 			{
 				UControlRig* RealControlRig = Cast<UControlRig>(ControlRig.Get());
-				GizmoTransform = UControlRigSequencerEditorLibrary::GetControlRigWorldTransform(LevelSequence, RealControlRig, Name, FrameTime.RoundToFrame(),
-					EMovieSceneTimeUnit::TickResolution);
+				GizmoTransform  = UE::AIE::PivotTool::GetControlWorldSpaceTransform(RealControlRig, Name);
+
 				AverageLocation += GizmoTransform.GetLocation();
 				NumLocations++;
 			}
@@ -638,7 +707,7 @@ void USequencerPivotTool::GizmoTransformStarted(UTransformProxy* Proxy)
 			{
 				continue;
 			}
-			TArray<FName> SelectedControls = ControlRig->CurrentControlSelection();
+			TArray<FName> SelectedControls = UE::AIE::PivotTool::CurrentControlSelection(ControlRig.Get());
 			if (SelectedControls.Num() > 0)
 			{
 				ControlRig->Modify();
@@ -661,8 +730,7 @@ void USequencerPivotTool::GizmoTransformStarted(UTransformProxy* Proxy)
 							}
 						}
 						UControlRig* RealControlRig = Cast<UControlRig>(ControlRig.Get());
-						FTransform Transform = UControlRigSequencerEditorLibrary::GetControlRigWorldTransform(LevelSequence, RealControlRig, Name, FrameNumber,
-							EMovieSceneTimeUnit::TickResolution);
+						FTransform Transform = UE::AIE::PivotTool::GetControlWorldSpaceTransform(RealControlRig, Name);
 
 						FControlRigSelectionDuringDrag ControlDrag;
 						ControlDrag.LevelSequence = LevelSequence;
@@ -1025,13 +1093,16 @@ void USequencerPivotTool::OnClicked(const FInputDeviceRay& ClickPos)
 FInputRayHit USequencerPivotTool::FindRayHit(const FRay& WorldRay, FVector& HitPos)
 {
 	// trace a ray into the World
-	FCollisionObjectQueryParams QueryParams(FCollisionObjectQueryParams::AllObjects);
-	FHitResult Result;
-	bool bHitWorld = TargetWorld->LineTraceSingleByObjectType(Result, WorldRay.Origin, WorldRay.PointAt(999999), QueryParams);
-	if (bHitWorld)
+	if (TargetWorld)
 	{
-		HitPos = Result.ImpactPoint;
-		return FInputRayHit(Result.Distance);
+		FCollisionObjectQueryParams QueryParams(FCollisionObjectQueryParams::AllObjects);
+		FHitResult Result;
+		bool bHitWorld = TargetWorld->LineTraceSingleByObjectType(Result, WorldRay.Origin, WorldRay.PointAt(999999), QueryParams);
+		if (bHitWorld)
+		{
+			HitPos = Result.ImpactPoint;
+			return FInputRayHit(Result.Distance);
+		}
 	}
 	return FInputRayHit();
 }
@@ -1363,7 +1434,6 @@ void USequencerPivotTool::UpdatePivotOverlayLocation(const FVector2D InLocation,
 	LastPivotOverlayLocation = ScreenPos;
 }
 
-
 #undef LOCTEXT_NAMESPACE
 
 
