Motion Trail Fixes, Add selection range option for drawing and, eval range, fix spaces with names in toolbar, make little faster by evaluating more on first tick,mark color and current mark color, and improve overlay mode by making it cache evaluations so when moving the offset we don't have to recalucate the evaluation, and fixing offset calculation to move locally.

diff --git a/Engine/Plugins/Animation/ControlRig/Source/ControlRigEditor/Private/Editor/ControlRigViewportToolbarExtensions.cpp b/Engine/Plugins/Animation/ControlRig/Source/ControlRigEditor/Private/Editor/ControlRigViewportToolbarExtensions.cpp
--- a/Engine/Plugins/Animation/ControlRig/Source/ControlRigEditor/Private/Editor/ControlRigViewportToolbarExtensions.cpp
+++ b/Engine/Plugins/Animation/ControlRig/Source/ControlRigEditor/Private/Editor/ControlRigViewportToolbarExtensions.cpp
@@ -128,7 +128,7 @@ void CreatePinnedMenuDelegate(UToolMenu* SubMenu, UMotionTrailToolOptions* Setti
 		LOCTEXT("OffsetLabelTooltip", "Toggle offset on selects the curve in the viewport, and allows you to move it like shift select does. Toggling it off will remove any offset."),
 		FSlateIcon(),
 		OffsetAction,
-		EUserInterfaceActionType::ToggleButton
+		EUserInterfaceActionType::Check
 	);
 	Section.AddEntry(OffsetEntry);
 
@@ -142,7 +142,7 @@ void CreatePinnedMenuDelegate(UToolMenu* SubMenu, UMotionTrailToolOptions* Setti
 		LOCTEXT("SpaceLabelTooltip", "Toggling on space will put you into eye drop selection mode to pick the scene compponent/socket that you want to have this trail in. Toggling it off puts it back in world space."),
 		FSlateIcon(),
 		SpaceAction,
-		EUserInterfaceActionType::ToggleButton
+		EUserInterfaceActionType::Check
 	);
 	Section.AddEntry(SpaceEntry);
 }
@@ -205,7 +205,7 @@ TSharedRef<SWidget> CreateFramesBeforeWidget(UMotionTrailToolOptions* Settings)
 						.Padding(FMargin(1.0f))
 						[
 							SNew(SSpinBox<int32>)
-								.IsEnabled_Lambda([Settings]() { return !Settings->bShowFullTrail; })
+								.IsEnabled_Lambda([Settings]() { return (Settings->MotionTrailRange == EMotionTrailRange::SpecifiedRange); })
 								.MinValue(0)
 								.MaxValue(100)
 								.MinDesiredWidth(50)
@@ -260,7 +260,7 @@ TSharedRef<SWidget> CreateFramesAfterWidget(UMotionTrailToolOptions* Settings)
 						.Padding(FMargin(1.0f))
 						[
 							SNew(SSpinBox<int32>)
-								.IsEnabled_Lambda([Settings]() { return !Settings->bShowFullTrail; })
+								.IsEnabled_Lambda([Settings]() { return (Settings->MotionTrailRange == EMotionTrailRange::SpecifiedRange);  })
 								.MinValue(0)
 								.MaxValue(100)
 								.MinDesiredWidth(50)
@@ -608,10 +608,16 @@ TSharedRef<SWidget> CreateTrailColorWidget(UMotionTrailToolOptions* Settings, FN
 	// clang-format on
 }
 
-FToolMenuEntry CreateTrailColor(UMotionTrailToolOptions* Settings, FName PropertyName)
+FToolMenuEntry CreateTrailColor(UMotionTrailToolOptions* Settings, FName PropertyName, FName Label, FName ToolTip)
 {
-	FText Text = FText::FromString(PropertyName.ToString());
-	FToolMenuEntry Entry = FToolMenuEntry::InitWidget(PropertyName, CreateTrailColorWidget(Settings, PropertyName), Text);
+	FText Text = FText::FromString(Label.ToString());
+	FText TooltipTex = FText::FromString(ToolTip.ToString());
+	
+	constexpr bool bNoIndent = false;
+	constexpr bool bSearchable = true;
+	constexpr bool bNoPadding = false;
+	FToolMenuEntry Entry = FToolMenuEntry::InitWidget(PropertyName, CreateTrailColorWidget(Settings, PropertyName), Text, bNoIndent, bSearchable,
+		bNoPadding,TooltipTex);
 
 	return Entry;
 }
@@ -670,9 +676,9 @@ TSharedRef<SWidget> CreatePropertyWidget(UObject* Settings, FName PropertyName)
 }
 
 template<typename NumericType>
-FToolMenuEntry CreateProperty(UObject* Settings, FName PropertyName)
+FToolMenuEntry CreateProperty(UObject* Settings, FName PropertyName, FName Label)
 {
-	FText Text = FText::FromString(PropertyName.ToString());
+	FText Text = FText::FromString(Label.ToString());
 	FToolMenuEntry Entry = FToolMenuEntry::InitWidget(PropertyName, CreatePropertyWidget<NumericType>(Settings, PropertyName), Text);
 
 	return Entry;
@@ -744,46 +750,75 @@ void CreateAdvancedSubMenu(UToolMenu* InSubMenu, UMotionTrailToolOptions* Settin
 	}
 	{
 		const FName DoubleProperty("KeySize");
-		TrailSettings.AddEntry(CreateProperty<double>(Settings, DoubleProperty));
+		const FName Label("Key Size");
+		TrailSettings.AddEntry(CreateProperty<double>(Settings, DoubleProperty, Label));
 	}
 
 	{
 		const FName DoubleProperty("MarkSize");
-		TrailSettings.AddEntry(CreateProperty<double>(Settings, DoubleProperty));
+		const FName Label("Mark Size");
+		TrailSettings.AddEntry(CreateProperty<double>(Settings, DoubleProperty, Label));
 	}
 	{
 		const FName DoubleProperty("TrailThickness");
-		TrailSettings.AddEntry(CreateProperty<double>(Settings, DoubleProperty));
+		const FName Label("Trail Thickness");
+		TrailSettings.AddEntry(CreateProperty<double>(Settings, DoubleProperty, Label));
 	}
 
 	FToolMenuSection& ColorSettings = InSubMenu->AddSection("ColorSettings", LOCTEXT("ColorSettings", "Color Settings"));
 	{
 		const FName ColorProperty("DefaultColor");
-		ColorSettings.AddEntry(CreateTrailColor(Settings, ColorProperty));
+		const FName Label("Default Color");
+		const FName Tooltip("Color for default style");
+		ColorSettings.AddEntry(CreateTrailColor(Settings, ColorProperty, Label, Tooltip));
 	}
 	{
 		const FName ColorProperty("TimePreColor");
-		ColorSettings.AddEntry(CreateTrailColor(Settings, ColorProperty));
+		const FName Label("Time Previous Color");
+		const FName Tooltip("Color before current time for time style");
+		ColorSettings.AddEntry(CreateTrailColor(Settings, ColorProperty, Label, Tooltip));
 	}
 	{
 		const FName ColorProperty("TimePostColor");
-		ColorSettings.AddEntry(CreateTrailColor(Settings, ColorProperty));
+		const FName Label("Time Post Color");
+		const FName Tooltip("Color after current time for time style");
+		ColorSettings.AddEntry(CreateTrailColor(Settings, ColorProperty, Label, Tooltip));
 	}
 	{
 		const FName ColorProperty("DashPreColor");
-		ColorSettings.AddEntry(CreateTrailColor(Settings, ColorProperty));
+		const FName Label("Dash Previous Color");
+		const FName Tooltip("Previous Color for dash style");
+		ColorSettings.AddEntry(CreateTrailColor(Settings, ColorProperty, Label, Tooltip));
 	}
 	{
 		const FName ColorProperty("DashPostColor");
-		ColorSettings.AddEntry(CreateTrailColor(Settings, ColorProperty));
+		const FName Label("Dash Post Color");
+		const FName Tooltip("Next Color for dash style");
+		ColorSettings.AddEntry(CreateTrailColor(Settings, ColorProperty, Label, Tooltip));
 	}
 	{
 		const FName ColorProperty("KeyColor");
-		ColorSettings.AddEntry(CreateTrailColor(Settings, ColorProperty));
+		const FName Label("Key Color");
+		const FName Tooltip("Color for keys");
+		ColorSettings.AddEntry(CreateTrailColor(Settings, ColorProperty, Label, Tooltip));
 	}
 	{
 		const FName ColorProperty("SelectedKeyColor");
-		ColorSettings.AddEntry(CreateTrailColor(Settings, ColorProperty));
+		const FName Label("Selected Key Color");
+		const FName Tooltip("Color for selected keys");
+		ColorSettings.AddEntry(CreateTrailColor(Settings, ColorProperty, Label, Tooltip));
+	}
+	{
+		const FName ColorProperty("MarkColor");
+		const FName Label("Mark Color");
+		const FName Tooltip("Color for marks");
+		ColorSettings.AddEntry(CreateTrailColor(Settings, ColorProperty, Label, Tooltip));
+	}
+	{
+		const FName ColorProperty("CurrentFrameMarkColor");
+		const FName Label("Current Mark Color");
+		const FName Tooltip("Color for mark at current frame");
+		ColorSettings.AddEntry(CreateTrailColor(Settings, ColorProperty, Label, Tooltip));
 	}
 
 	FToolMenuSection& PinSettings = InSubMenu->AddSection("PinSettings", LOCTEXT("PinSettings", "Pin Settings"));
@@ -811,23 +846,53 @@ void CreateMotionTrailMenu(UToolMenu* InMenu)
 		Action.ExecuteAction = FExecuteAction::CreateLambda(
 			[Settings]() -> void
 			{
-				Settings->bShowFullTrail = !Settings->bShowFullTrail;
-				FPropertyChangedEvent Event(UMotionTrailToolOptions::StaticClass()->FindPropertyByName(GET_MEMBER_NAME_CHECKED(UMotionTrailToolOptions, bShowFullTrail)));
+				Settings->MotionTrailRange = EMotionTrailRange::PlaybackRange;
+				FPropertyChangedEvent Event(UMotionTrailToolOptions::StaticClass()->FindPropertyByName(GET_MEMBER_NAME_CHECKED(UMotionTrailToolOptions, MotionTrailRange)));
+				Settings->PostEditChangeProperty(Event);
+			}
+		);
+		Action.GetActionCheckState = FGetActionCheckState::CreateLambda(
+			[Settings]() -> ECheckBoxState
+			{
+				return (Settings->MotionTrailRange == EMotionTrailRange::PlaybackRange) ? ECheckBoxState::Checked
+					: ECheckBoxState::Unchecked;
+			}
+		);
+
+		FToolMenuEntry Entry = FToolMenuEntry::InitMenuEntry(
+			"PlaybackRange",
+			LOCTEXT("PlaybackRangeLabel", "Playback Range"),
+			LOCTEXT("PlaybackRangeTooltip", "Show full trail along playback range"),
+			FSlateIcon(),
+			Action,
+			EUserInterfaceActionType::RadioButton
+		);
+		PathModeSection.AddEntry(Entry);
+	}
+	//show selected
+	{
+		FUIAction Action;
+
+		Action.ExecuteAction = FExecuteAction::CreateLambda(
+			[Settings]() -> void
+			{
+				Settings->MotionTrailRange = EMotionTrailRange::SelectionRange;
+				FPropertyChangedEvent Event(UMotionTrailToolOptions::StaticClass()->FindPropertyByName(GET_MEMBER_NAME_CHECKED(UMotionTrailToolOptions, MotionTrailRange)));
 				Settings->PostEditChangeProperty(Event);
 			}
 		);
 		Action.GetActionCheckState = FGetActionCheckState::CreateLambda(
 			[Settings]() -> ECheckBoxState
 			{
-				return Settings->bShowFullTrail ? ECheckBoxState::Checked
+				return (Settings->MotionTrailRange == EMotionTrailRange::SelectionRange) ? ECheckBoxState::Checked
 					: ECheckBoxState::Unchecked;
 			}
 		);
 
 		FToolMenuEntry Entry = FToolMenuEntry::InitMenuEntry(
-			"FullTrail",
-			LOCTEXT("FullTrailLabel", "Full trail"),
-			LOCTEXT("FullTrailTooltip", "Show full trail"),
+			"SelectionRange",
+			LOCTEXT("SelectionRangeLabel", "Selection Range"),
+			LOCTEXT("FullTrailTooltip", "Show trail in the selection range"),
 			FSlateIcon(),
 			Action,
 			EUserInterfaceActionType::RadioButton
@@ -841,15 +906,15 @@ void CreateMotionTrailMenu(UToolMenu* InMenu)
 		Action.ExecuteAction = FExecuteAction::CreateLambda(
 			[Settings]() -> void
 			{
-				Settings->bShowFullTrail = !Settings->bShowFullTrail;
-				FPropertyChangedEvent Event(UMotionTrailToolOptions::StaticClass()->FindPropertyByName(GET_MEMBER_NAME_CHECKED(UMotionTrailToolOptions, bShowFullTrail)));
+				Settings->MotionTrailRange = EMotionTrailRange::SpecifiedRange;
+				FPropertyChangedEvent Event(UMotionTrailToolOptions::StaticClass()->FindPropertyByName(GET_MEMBER_NAME_CHECKED(UMotionTrailToolOptions, MotionTrailRange)));
 				Settings->PostEditChangeProperty(Event);
 			}
 		);
 		Action.GetActionCheckState = FGetActionCheckState::CreateLambda(
 			[Settings]() -> ECheckBoxState
 			{
-				return Settings->bShowFullTrail == false ? ECheckBoxState::Checked
+				return Settings->MotionTrailRange == EMotionTrailRange::SpecifiedRange ? ECheckBoxState::Checked
 					: ECheckBoxState::Unchecked;
 			}
 		);
@@ -1211,7 +1276,8 @@ void CreateAxisOnSelectionMenu(UToolMenu* AnimationShowFlagsSubmenu, FToolMenuSe
 				Submenu->FindOrAddSection("AxisOnSelection", LOCTEXT("AxisOnSelectionLabel", "Axis On Selection"));
 
 			const FName DoubleProperty("AxisScale");
-			FToolMenuEntry NumericEntry = CreateProperty<float>(Settings, DoubleProperty);
+			const FName Label("Axis Scale");
+			FToolMenuEntry NumericEntry = CreateProperty<float>(Settings, DoubleProperty, Label);
 
 			AxisOnSelection.AddEntry(NumericEntry);
 		}
diff --git a/Engine/Plugins/Editor/SequencerAnimTools/Source/SequencerAnimTools/Private/MotionTrailTool.cpp b/Engine/Plugins/Editor/SequencerAnimTools/Source/SequencerAnimTools/Private/MotionTrailTool.cpp
--- a/Engine/Plugins/Editor/SequencerAnimTools/Source/SequencerAnimTools/Private/MotionTrailTool.cpp
+++ b/Engine/Plugins/Editor/SequencerAnimTools/Source/SequencerAnimTools/Private/MotionTrailTool.cpp
@@ -307,19 +307,22 @@ void UMotionTrailTool::GizmoTransformChanged(UTransformProxy* Proxy, FTransform
 	}
 	GizmoTransform = Transform;
 	FTransform Diff = Transform.GetRelativeTransform(StartDragTransform);
-	FVector LocationDiff = GizmoTransform.GetLocation() - StartDragTransform.GetLocation();
-	for (const TUniquePtr<UE::SequencerAnimTools::FTrailHierarchy>& TrailHierarchy : TrailHierarchies)
+	if (Diff.Equals(FTransform::Identity, UE_KINDA_SMALL_NUMBER) == false)
 	{
-		FVector WidgetLocation = Transform.GetLocation();
-		const bool bShiftIsDown = FSlateApplication::Get().GetModifierKeys().IsShiftDown();
-		bool bIsHandled = TrailHierarchy->ApplyDelta(LocationDiff, Diff.GetRotation().Rotator(), WidgetLocation, bShiftIsDown);
-		if (bIsHandled)
+		FVector LocationDiff = GizmoTransform.GetLocation() - StartDragTransform.GetLocation();
+		for (const TUniquePtr<UE::SequencerAnimTools::FTrailHierarchy>& TrailHierarchy : TrailHierarchies)
 		{
-			bManipulatorMadeChange = true;
+			FVector WidgetLocation = Transform.GetLocation();
+			const bool bShiftIsDown = FSlateApplication::Get().GetModifierKeys().IsShiftDown();
+			bool bIsHandled = TrailHierarchy->ApplyDelta(LocationDiff, GizmoTransform.GetRotation().Rotator(), WidgetLocation, bShiftIsDown);
+			if (bIsHandled)
+			{
+				bManipulatorMadeChange = true;
+			}
 		}
+		StartDragTransform = Transform;
+		UpdateGizmoLocation();
 	}
-	StartDragTransform = Transform;
-	UpdateGizmoLocation();
 
 }
 
@@ -375,24 +378,33 @@ void UMotionTrailTool::UpdateGizmoLocation()
 	{
 		FVector NewGizmoLocation(0., 0., 0.);
 		int NumSelected = 0;
+		TOptional<FQuat> NewRotation;
+		FQuat LocalRotation;
 		for (const TUniquePtr<UE::SequencerAnimTools::FTrailHierarchy>& TrailHierarchy : TrailHierarchies)
 		{
 			FVector Location;
-			bool bIsHandled = TrailHierarchy->IsAnythingSelected(Location);
+			
+			bool bIsHandled = TrailHierarchy->IsAnythingSelected(Location,LocalRotation);
 			if (bIsHandled)
 			{
 				NewGizmoLocation += Location;
+				if (NewRotation.IsSet() == false)
+				{
+					NewRotation = LocalRotation;
+				}
 				++NumSelected;
 			}
 		}
 		if (NumSelected > 0)
 		{
 			NewGizmoLocation /= (double)NumSelected;
+			LocalRotation = NewRotation.GetValue();
 		}
-		return NewGizmoLocation;
+		return TPair<FVector, FQuat>(NewGizmoLocation, LocalRotation);
 	};
-	
-	GizmoTransform.SetLocation(GetNewLocation());
+	TPair<FVector, FQuat> LocRot = GetNewLocation();
+	GizmoTransform.SetLocation(LocRot.Key);
+	GizmoTransform.SetRotation(LocRot.Value);;
 
 	if (TransformGizmo)
 	{
diff --git a/Engine/Plugins/Editor/SequencerAnimTools/Source/SequencerAnimTools/Private/Sequencer/MotionTrailMovieSceneKey.cpp b/Engine/Plugins/Editor/SequencerAnimTools/Source/SequencerAnimTools/Private/Sequencer/MotionTrailMovieSceneKey.cpp
--- a/Engine/Plugins/Editor/SequencerAnimTools/Source/SequencerAnimTools/Private/Sequencer/MotionTrailMovieSceneKey.cpp
+++ b/Engine/Plugins/Editor/SequencerAnimTools/Source/SequencerAnimTools/Private/Sequencer/MotionTrailMovieSceneKey.cpp
@@ -100,16 +100,16 @@ void FMotionTraiMovieScenelKeyTool::Render(const FGuid& Guid, const FSceneView*
 
 void FMotionTraiMovieScenelKeyTool::DrawHUD(const FSceneView* View, FCanvas* Canvas)
 {
-	return;
-	
+	return;	
 }
 
-bool FMotionTraiMovieScenelKeyTool::IsSelected(FVector& OutVectorPosition) const
+bool FMotionTraiMovieScenelKeyTool::IsSelected(FVector& OutVectorPosition, FQuat& OutRotation) const
 {
 	if (CachedSelection.Num() > 0)
 	{
 		UpdateSelectedKeysTransform();
 		OutVectorPosition = SelectedKeysTransform.GetLocation();
+		OutRotation = SelectedKeysTransform.GetRotation();
 		return true;
 	}
 	return false;
@@ -185,6 +185,8 @@ void FMotionTraiMovieScenelKeyTool::UpdateSelectedKeysTransform() const
 {
 	if (CachedSelection.Num() > 0)
 	{
+
+		TOptional<FQuat> NewRotation;
 		const FTransform OffsetTransform = FTransform::Identity;//OwningTrail->GetOffsetTransform();
 		const FTransform ParentSpaceTransform = OwningTrail->GetParentSpaceTransform();
 		FVector NewGizmoLocation = FVector::ZeroVector;
@@ -192,9 +194,14 @@ void FMotionTraiMovieScenelKeyTool::UpdateSelectedKeysTransform() const
 		{
 			const FTransform Transform = KeyInfo->Transform * OffsetTransform * ParentSpaceTransform;
 			NewGizmoLocation += Transform.GetLocation();
+			if (NewRotation.IsSet() == false)
+			{
+				NewRotation = Transform.GetRotation();
+			}
 		}
 		NewGizmoLocation /= (double)CachedSelection.Num();
 		SelectedKeysTransform.SetLocation(NewGizmoLocation);
+		SelectedKeysTransform.SetRotation(NewRotation.GetValue());
 	}
 }
 
diff --git a/Engine/Plugins/Editor/SequencerAnimTools/Source/SequencerAnimTools/Private/Sequencer/MotionTrailMovieSceneKey.h b/Engine/Plugins/Editor/SequencerAnimTools/Source/SequencerAnimTools/Private/Sequencer/MotionTrailMovieSceneKey.h
--- a/Engine/Plugins/Editor/SequencerAnimTools/Source/SequencerAnimTools/Private/Sequencer/MotionTrailMovieSceneKey.h
+++ b/Engine/Plugins/Editor/SequencerAnimTools/Source/SequencerAnimTools/Private/Sequencer/MotionTrailMovieSceneKey.h
@@ -220,7 +220,7 @@ public:
 
 	bool HandleClick(const FGuid& Guid, FEditorViewportClient* InViewportClient, HHitProxy* HitProxy, FInputClick Click);
 	void StartDragging() {};
-	bool IsSelected(FVector& OutVectorPosition) const;
+	bool IsSelected(FVector& OutVectorPosition, FQuat& OutRotation) const;
 	void GetSelectedKeyPositions(TArray<FVector>& OutPositions) const;
 	bool IsSelected() const;
 	void OnSectionChanged();
diff --git a/Engine/Plugins/Editor/SequencerAnimTools/Source/SequencerAnimTools/Private/Sequencer/MovieSceneTransformTrail.cpp b/Engine/Plugins/Editor/SequencerAnimTools/Source/SequencerAnimTools/Private/Sequencer/MovieSceneTransformTrail.cpp
--- a/Engine/Plugins/Editor/SequencerAnimTools/Source/SequencerAnimTools/Private/Sequencer/MovieSceneTransformTrail.cpp
+++ b/Engine/Plugins/Editor/SequencerAnimTools/Source/SequencerAnimTools/Private/Sequencer/MovieSceneTransformTrail.cpp
@@ -125,6 +125,7 @@ FMovieSceneTransformTrail::FMovieSceneTransformTrail(const FGuid& InBindingID, U
 	KeyTool = MakeUnique< FMotionTraiMovieScenelKeyTool>(this);
 
 	//new
+	CalculatedArrayOfTransforms = MakeShared<FArrayOfTransforms>();
 	ArrayOfTransforms = MakeShared<FArrayOfTransforms>();
 	ParentArrayOfTransforms = MakeShared<FArrayOfTransforms>();
 
@@ -144,7 +145,15 @@ void FMovieSceneTransformTrail::HasStartedEvaluating()
 void FMovieSceneTransformTrail::UpdateFinished(const TRange<FFrameNumber>& Range, const TArray<int32>& IndicesToCalcluate, bool bDoneCalculating)
 {
 	FTrail::UpdateFinished(Range, IndicesToCalcluate, bDoneCalculating);
-	if (HasOffsetTransform())
+	if (HasOffsetTransform() == false)
+	{
+		for (int32 Count = 0; Count < IndicesToCalcluate.Num(); ++Count)
+		{
+			int32 Index = IndicesToCalcluate[Count];
+			ArrayOfTransforms->Transforms[Index] = CalculatedArrayOfTransforms->Transforms[Index];
+		}
+	}
+	else
 	{
 		TSet<int32> CalculatedIndices;
 		for (int32 Count = 0; Count < IndicesToCalcluate.Num(); ++Count)
@@ -153,7 +162,7 @@ void FMovieSceneTransformTrail::UpdateFinished(const TRange<FFrameNumber>& Range
 			if (CalculatedIndices.Contains(Index) == false)
 			{
 				CalculatedIndices.Add(Index);
-				ArrayOfTransforms->Transforms[Index] = OffsetTransform * ArrayOfTransforms->Transforms[Index];
+				ArrayOfTransforms->Transforms[Index] = OffsetTransform * CalculatedArrayOfTransforms->Transforms[Index];
 			}
 		}
 	}
@@ -432,6 +441,7 @@ FTrailCurrentStatus FMovieSceneTransformTrail::UpdateTrail(const FNewSceneContex
 		bForceEvaluateNextTick = false;
 
 		UpdateNeedsEvaluation(InSceneContext.YourNode, SequencerTrailHierarchy);
+
 	}
 	else
 	{
@@ -588,7 +598,7 @@ TOptional<FLinearColor> FMovieSceneTransformTrail::GetOverrideColor() const
 
 void FMovieSceneTransformTrail::InternalDrawTrail(const FGuid& Guid, const FSceneView* View, FPrimitiveDrawInterface* PDI, const FDrawCacheData& InDrawData, bool bHitTesting, bool bTrailIsEvaluating) const
 {
-	if (bTrailIsEvaluating == false) // if evaluating we now just draw betweeen the keys
+	if (bTrailIsEvaluating == false || bIsOffsetMode == true) // if evaluating we now just draw betweeen the keys
 	{
 		const TOptional<FLinearColor> OverrideColor = GetOverrideColor();
 		const float TrailThickness = UMotionTrailToolOptions::GetTrailOptions()->TrailThickness;
@@ -620,6 +630,10 @@ void FMovieSceneTransformTrail::Render(const FGuid& Guid, const FSceneView* View
 		const bool bHitTesting = PDI && PDI->IsHitTesting();
 		InternalDrawTrail(Guid, View, PDI, CachedDrawData, bHitTesting, bTrailIsEvaluating);
 	}
+	else if (bIsOffsetMode && bTrailIsEvaluating)
+	{
+		RenderEvaluating(Guid, View, PDI);
+	}
 	//render key tool
 	KeyTool->Render(Guid, View, PDI, bTrailIsEvaluating);
 }
@@ -629,7 +643,7 @@ void FMovieSceneTransformTrail::RenderEvaluating(const FGuid& Guid, const FScene
 	if (PreviousCachedDrawData.PointsToDraw.Num() > 0)
 	{
 		const bool bHitTesting = false;
-		InternalDrawTrail(Guid, View, PDI, PreviousCachedDrawData, bHitTesting,false); //if rendering cached we say not evaluating
+		InternalDrawTrail(Guid, View, PDI, PreviousCachedDrawData, bHitTesting, false); //if rendering cached we say not evaluating
 	}
 }
 
@@ -687,6 +701,7 @@ bool FMovieSceneTransformTrail::HandleClick(const FGuid& Guid, FEditorViewportCl
 				bIsSelected = true;
 				SelectedPos = HitProxy->Point;
 			}
+			SelectedRotation = FQuat::Identity;
 			return true;
 		}
 	}
@@ -725,14 +740,15 @@ bool FMovieSceneTransformTrail::IsAnythingSelected(TArray<FVector>& OutVectorPos
 	return false;
 }
 
-bool FMovieSceneTransformTrail::IsAnythingSelected(FVector& OutVectorPosition) const
+bool FMovieSceneTransformTrail::IsAnythingSelected(FVector& OutVectorPosition, FQuat& OutRotation) const
 {
 	if (bIsSelected)
 	{
 		OutVectorPosition = SelectedPos;
+		OutRotation = SelectedRotation;
 		return true;
 	}
-	return KeyTool->IsSelected(OutVectorPosition);
+	return KeyTool->IsSelected(OutVectorPosition, OutRotation);
 }
 
 void FMovieSceneTransformTrail::TranslateSelectedKeys(bool bRight)
@@ -779,9 +795,14 @@ void FMovieSceneTransformTrail::ClearOffsetTransform()
 	{
 		ExitOffsetMode();
 		bIsSelected = false;
+		CalculateUpdatedOffset();
 		ClearCachedData();
 	}
 }
+bool FMovieSceneTransformTrail::GetOffsetMode() const
+{
+	return bIsOffsetMode;
+}
 
 void FMovieSceneTransformTrail::SetOffsetMode()
 {
@@ -799,13 +820,15 @@ void FMovieSceneTransformTrail::SetOffsetMode()
 		FFrameTime StartTime = GetSequencer()->GetLocalTime().Time;
 		const FTransform Transform = ArrayOfTransforms->Interp(StartTime.GetFrame(), CurrentFramesInfo->TransformIndices, CurrentFramesInfo->CurrentFrames);
 		SelectedPos = Transform.GetLocation();
+		SelectedRotation = Transform.GetRotation();
+		CalculateUpdatedOffset();
 	}
 
 }
 
 bool FMovieSceneTransformTrail::IsTracking() const 
 { 
-	return bIsTracking; 
+	return bIsTracking;
 }
 
 bool FMovieSceneTransformTrail::StartTracking()
@@ -822,20 +845,52 @@ bool FMovieSceneTransformTrail::StartTracking()
 	return false;
 }
 
+void FMovieSceneTransformTrail::CalculateUpdatedOffset()
+{
+	if (ArrayOfTransforms->Num() == CalculatedArrayOfTransforms->Num())
+	{
+		ClearCachedData();  
+		CacheState = ETrailCacheState::UpToDate;
+		for (int32 Index = 0; Index < ArrayOfTransforms->Num(); ++Index)
+		{
+			FVector InOffsetSpace = CalculatedArrayOfTransforms->Transforms[Index].GetRotation().RotateVector(OffsetTransform.GetLocation());
+			InOffsetSpace += CalculatedArrayOfTransforms->Transforms[Index].GetLocation();
+			FTransform NewOffset(CalculatedArrayOfTransforms->Transforms[Index].GetRotation().Rotator(),
+				InOffsetSpace, CalculatedArrayOfTransforms->Transforms[Index].GetScale3D());
+			ArrayOfTransforms->Transforms[Index] = OffsetTransform * CalculatedArrayOfTransforms->Transforms[Index];
+		}
+		if (ParentSpace.bIsValid && ParentArrayOfTransforms->Num() == ArrayOfTransforms->Num())
+		{
+			TSet<int32> CalculatedIndices;
+			for (int32 Index = 0; Index < ArrayOfTransforms->Num(); ++Index)
+			{
+				ArrayOfTransforms->Transforms[Index] = ArrayOfTransforms->Transforms[Index].GetRelativeTransform(ParentSpace.ParentSpace.WorldTransforms->Transforms[Index]);
+				ParentArrayOfTransforms->Transforms[Index] = ParentArrayOfTransforms->Transforms[Index].GetRelativeTransform(ParentSpace.ParentSpace.WorldTransforms->Transforms[Index]);
+			}
+		}
+		KeyTool->DirtyKeyTransforms();
+		KeyTool->UpdateKeys();
+	}
+}
 bool FMovieSceneTransformTrail::ApplyDelta(const FVector& Pos, const FRotator& Rot, const FVector& WidgetLocation, bool bApplyToOffset)
 {
 	if (bApplyToOffset || bIsOffsetMode)
 	{
 		if (IsTrailSelected())
 		{
-			OffsetTransform.SetLocation(OffsetTransform.GetLocation() + Pos);
-			int32 Index = UMotionTrailToolOptions::GetTrailOptions()->GetIndexFromGuid(CachedHierarchyGuid);
-			if (Index != INDEX_NONE)
+			FQuat Rotation(Rot);
+			FVector UnRotatedVector = Rotation.UnrotateVector(Pos);
+			OffsetTransform.SetLocation(OffsetTransform.GetLocation() + UnRotatedVector);
+			if (bIsOffsetMode == false)
 			{
-				UMotionTrailToolOptions::GetTrailOptions()->SetHasOffset(Index,true);
+				int32 Index = UMotionTrailToolOptions::GetTrailOptions()->GetIndexFromGuid(CachedHierarchyGuid);
+				if (Index != INDEX_NONE)
+				{
+					UMotionTrailToolOptions::GetTrailOptions()->SetHasOffset(Index, true);
+				}
 			}
 			SelectedPos += Pos;
-			ForceEvaluateNextTick();
+			CalculateUpdatedOffset();
 			return true;
 		}
 	}
@@ -925,7 +980,9 @@ void FMovieSceneComponentTransformTrail::UpdateNeedsEvaluation(const FGuid& InTr
 		FActorForWorldTransforms Actors;
 		Actors.Actor = Component->GetOwner();
 		Actors.Component = Component;
-		SequencerHierarchy->EvaluateActor(InTrailGuid, Actors, ArrayOfTransforms, ParentArrayOfTransforms);
+		SequencerHierarchy->EvaluateActor(InTrailGuid, Actors, CalculatedArrayOfTransforms, ParentArrayOfTransforms);
+		ArrayOfTransforms->Transforms.SetNum(CalculatedArrayOfTransforms->Transforms.Num());
+
 	}
 }
 
@@ -1156,7 +1213,8 @@ void FMovieSceneSocketTransformTrail::UpdateNeedsEvaluation(const FGuid& InTrail
 		Actors.Actor = Component->GetOwner();
 		Actors.Component = Component;
 		Actors.SocketName = SocketName;
-		SequencerHierarchy->EvaluateActor(InTrailGuid, Actors, ArrayOfTransforms, ParentArrayOfTransforms);
+		SequencerHierarchy->EvaluateActor(InTrailGuid, Actors, CalculatedArrayOfTransforms, ParentArrayOfTransforms);
+		ArrayOfTransforms->Transforms.SetNum(CalculatedArrayOfTransforms->Transforms.Num());
 	}
 }
 
@@ -1269,11 +1327,12 @@ bool FMovieSceneSocketTransformTrail::HandleClick(const FGuid& Guid, FEditorView
 }
 
 
-bool FMovieSceneSocketTransformTrail::IsAnythingSelected(FVector& OutVectorPosition) const
+bool FMovieSceneSocketTransformTrail::IsAnythingSelected(FVector& OutVectorPosition, FQuat& OutRotation) const
 {
 	if (bIsSelected)
 	{
 		OutVectorPosition = SelectedPos;
+		OutRotation = FQuat::Identity;
 		return true;
 	}
 	return false;
@@ -1344,7 +1403,8 @@ void FMovieSceneControlRigTransformTrail::UpdateNeedsEvaluation(const FGuid& InT
 	}
 	FMovieSceneTransformTrail::UpdateNeedsEvaluation(InTrailGuid, SequencerHierarchy);
 
-	SequencerHierarchy->EvaluateControlRig(InTrailGuid, ControlRig, ControlName, ArrayOfTransforms);
+	SequencerHierarchy->EvaluateControlRig(InTrailGuid, ControlRig, ControlName, CalculatedArrayOfTransforms);
+	ArrayOfTransforms->Transforms.SetNum(CalculatedArrayOfTransforms->Transforms.Num());
 }
 
 FText FMovieSceneControlRigTransformTrail::GetName() const
@@ -1642,4 +1702,3 @@ void FMovieSceneControlRigTransformTrail::FMovieSceneControlRigTransformTrailSta
 } // namespace UE
 
 #undef LOCTEXT_NAMESPACE
-
diff --git a/Engine/Plugins/Editor/SequencerAnimTools/Source/SequencerAnimTools/Private/Sequencer/MovieSceneTransformTrail.h b/Engine/Plugins/Editor/SequencerAnimTools/Source/SequencerAnimTools/Private/Sequencer/MovieSceneTransformTrail.h
--- a/Engine/Plugins/Editor/SequencerAnimTools/Source/SequencerAnimTools/Private/Sequencer/MovieSceneTransformTrail.h
+++ b/Engine/Plugins/Editor/SequencerAnimTools/Source/SequencerAnimTools/Private/Sequencer/MovieSceneTransformTrail.h
@@ -76,16 +76,16 @@ public:
 	virtual void UpdateFinished(const TRange<FFrameNumber>& UpdatedRange, const TArray<int32>& IndicesToCalcluate, bool bDoneCalculating) override;
 	virtual void AddImportantTimes(TSet<FFrameNumber>& InOutImportantTimes) override;
 	virtual FTransform GetOffsetTransform() const override { return OffsetTransform; }
-	virtual void ClearOffsetTransform();
-	virtual void SetOffsetMode(); 
+	virtual void ClearOffsetTransform() override;
+	virtual void SetOffsetMode() override;
+	virtual bool GetOffsetMode() const override;
 
-	virtual FTransform GetParentSpaceTransform() const;
+	virtual FTransform GetParentSpaceTransform() const override;
 	virtual void SetSpace(AActor* InActor, const FName& InComponentName) override;
 	virtual void ClearSpace() override { ParentSpace.ClearSpace(); 	ClearCachedData();}
 	virtual void ClearCachedData() override;
 	virtual void ReadyToDrawTrail(FColorState& ColorState, const FCurrentFramesInfo* InCurrentFramesInfo, bool bIsEvaluating, bool bIsPinned) override;
-	virtual void ActorChanged(AActor* InActor);
-	virtual void UpdateNeedsEvaluation(const FGuid& InTrailGuid, FSequencerTrailHierarchy* SequencerHierarchy);
+	virtual void ActorChanged(AActor* InActor) override;
 
 	virtual void Render(const FGuid& Guid, const FSceneView* View,  FPrimitiveDrawInterface* PDI, bool bTrailIsEvaluating) override;
 	virtual void RenderEvaluating(const FGuid& Guid, const FSceneView* View, FPrimitiveDrawInterface* PDI) override;
@@ -93,7 +93,7 @@ public:
 
 	virtual bool HandleClick(const FGuid& Guid, FEditorViewportClient* InViewportClient, HHitProxy* HitProxy, FInputClick Click) override;
 	virtual bool IsAnythingSelected() const override;
-	virtual bool IsAnythingSelected(FVector& OutVectorPosition)const override;
+	virtual bool IsAnythingSelected(FVector& OutVectorPosition, FQuat& OutRotation)const override;
 	virtual bool IsAnythingSelected(TArray<FVector>& OutVectorPositions)const override;
 	virtual bool IsTrailSelected() const override;
 	virtual bool StartTracking() override;
@@ -114,6 +114,9 @@ public:
 
 	// End FTrail interface
 	
+	// New Virtuals
+	virtual void UpdateNeedsEvaluation(const FGuid& InTrailGuid, FSequencerTrailHierarchy* SequencerHierarchy);
+
 	TSharedPtr<ISequencer> GetSequencer() const { return WeakSequencer.Pin(); }
 	FGuid GetCachedHierarchyGuid() const { return CachedHierarchyGuid; }
 	UMovieSceneSection* GetSection() const;
@@ -138,11 +141,13 @@ protected:
 	virtual void CheckAndUpdateObjects();
 	static TArray<UObject*> GetBoundObjects(TSharedPtr<ISequencer>& InSequencerPtr, const FGuid& InGuid);
 
+	void CalculateUpdatedOffset();
 protected:
 	//some interactive states/modes
 	bool bIsSelected;
 	bool bIsOffsetMode;
 	FVector SelectedPos;
+	FQuat SelectedRotation;
 	bool bIsTracking; 
 	TOptional<bool> bSetShowWidget;
 
@@ -165,7 +170,11 @@ protected:
 	UE::AIE::FSequencerTransformDependencies TransformDependencies;
 
 	FTransform OffsetTransform;
+	//transforms calculated from evaluating control rigs and actors
+	TSharedPtr<UE::AIE::FArrayOfTransforms> CalculatedArrayOfTransforms;
+	//calculated transforms plus any offset + parent
 	TSharedPtr<UE::AIE::FArrayOfTransforms> ArrayOfTransforms;
+	//parent space transforms (optional)
 	TSharedPtr<UE::AIE::FArrayOfTransforms> ParentArrayOfTransforms;
 	const FCurrentFramesInfo* CurrentFramesInfo;
 
@@ -265,7 +274,7 @@ private:
 
 	virtual bool HandleClick(const FGuid& Guid, FEditorViewportClient* InViewportClient, HHitProxy* HitProxy, FInputClick Click) override;
 	virtual bool IsAnythingSelected() const override;
-	virtual bool IsAnythingSelected(FVector& OutVectorPosition)const override;
+	virtual bool IsAnythingSelected(FVector& OutVectorPosition, FQuat& OutRotation)const override;
 
 protected:
 
diff --git a/Engine/Plugins/Editor/SequencerAnimTools/Source/SequencerAnimTools/Private/Sequencer/SequencerTrailHierarchy.cpp b/Engine/Plugins/Editor/SequencerAnimTools/Source/SequencerAnimTools/Private/Sequencer/SequencerTrailHierarchy.cpp
--- a/Engine/Plugins/Editor/SequencerAnimTools/Source/SequencerAnimTools/Private/Sequencer/SequencerTrailHierarchy.cpp
+++ b/Engine/Plugins/Editor/SequencerAnimTools/Source/SequencerAnimTools/Private/Sequencer/SequencerTrailHierarchy.cpp
@@ -24,6 +24,9 @@
 #include "Engine/SkeletalMesh.h"
 #include "Selection.h"
 #include "Editor/UnrealEdEngine.h"
+#include "Editor.h"
+#include "EditorModeManager.h"
+
 
 namespace UE
 {
@@ -95,6 +98,8 @@ FSequencerTrailHierarchy::FSequencerTrailHierarchy(TWeakPtr<ISequencer> InWeakSe
 	, OnSelectionChangedHandle()
 	, OnViewOptionsChangedHandle()
 	, ControlRigDelegateHandles()
+	, SavedWidgetMode(UE::Widget::WM_Translate)
+	, SavedCoordSystem(COORD_Local)
 {
 }
 void FSequencerTrailHierarchy::Initialize()
@@ -104,7 +109,11 @@ void FSequencerTrailHierarchy::Initialize()
 	{
 		return;
 	}
-
+	UMotionTrailToolOptions* Settings = UMotionTrailToolOptions::GetTrailOptions();
+	if (!Settings)
+	{
+		return;
+	}
 	TArray<FGuid> SequencerSelectedObjects;
 	Sequencer->GetSelectedObjects(SequencerSelectedObjects);
 	UpdateSequencerBindings(SequencerSelectedObjects,
@@ -136,7 +145,7 @@ void FSequencerTrailHierarchy::Initialize()
 
 	});
 
-	OnViewOptionsChangedHandle = UMotionTrailToolOptions::GetTrailOptions()->OnDisplayPropertyChanged.AddLambda([this](FName PropertyName) {
+	OnViewOptionsChangedHandle = Settings->OnDisplayPropertyChanged.AddLambda([this, Settings](FName PropertyName) {
 		if (PropertyName == GET_MEMBER_NAME_CHECKED(UMotionTrailToolOptions, EvalsPerFrame))
 		{
 			for (const TPair<FGuid, TUniquePtr<FTrail>>& GuidTrailPair : GetAllTrails())
@@ -159,7 +168,7 @@ void FSequencerTrailHierarchy::Initialize()
 		}
 		else if (PropertyName == GET_MEMBER_NAME_CHECKED(UMotionTrailToolOptions, bShowSelectedTrails))
 		{
-			if (UMotionTrailToolOptions::GetTrailOptions()->bShowSelectedTrails == false)
+			if (Settings->bShowSelectedTrails == false)
 			{
 				TArray<FGuid> TrailsToRemove; //if trail is not visible we just get rid of it
 				// <parent, current>
@@ -230,7 +239,10 @@ void FSequencerTrailHierarchy::Destroy()
 	{
 		Sequencer->OnActorAddedToSequencer().Remove(OnActorAddedToSequencerHandle);
 		Sequencer->GetSelectionChangedObjectGuids().Remove(OnSelectionChangedHandle);
-		UMotionTrailToolOptions::GetTrailOptions()->OnDisplayPropertyChanged.Remove(OnViewOptionsChangedHandle);
+		if (UMotionTrailToolOptions* Settings = UMotionTrailToolOptions::GetTrailOptions())
+		{
+			Settings->OnDisplayPropertyChanged.Remove(OnViewOptionsChangedHandle);
+		}
 	}
 
 	for (const TPair<UMovieSceneControlRigParameterTrack*, FControlRigDelegateHandles>& SectionHandlesPair : ControlRigDelegateHandles)
@@ -275,7 +287,10 @@ void FSequencerTrailHierarchy::Destroy()
 	UnRegisterMotionTrailOptionDelegates();
 	FCoreUObjectDelegates::OnObjectsReplaced.Remove(OnObjectsReplacedHandle);
 
-	UMotionTrailToolOptions::GetTrailOptions()->ResetPinnedItems();
+	if (UMotionTrailToolOptions* Settings = UMotionTrailToolOptions::GetTrailOptions())
+	{
+		Settings->ResetPinnedItems();
+	}
 }
 
 void FSequencerTrailHierarchy::OnActorChangedSomehow(AActor* InActor)
@@ -300,7 +315,7 @@ void FSequencerTrailHierarchy::OnActorsChangedSomehow(TArray<AActor*>& InActors)
 void FSequencerTrailHierarchy::CalculateEvalRangeArray()
 {
 	TicksPerSegment = GetFramesPerFrame();
-	CurrentFramesInfo.SetViewRange(TickViewRange);
+	CurrentFramesInfo.SetViewRange(TickViewRange, TickEvalRange == TickViewRange);
 	if (LastTicksPerSegment != TicksPerSegment || TickEvalRange != LastTickEvalRange)
 	{	
 		LastTicksPerSegment = TicksPerSegment;
@@ -311,6 +326,7 @@ void FSequencerTrailHierarchy::CalculateEvalRangeArray()
 		}
 		CurrentFramesInfo.SetUpFrameTimes(TickEvalRange, 
 			TicksPerSegment);
+		
 	}
 }
 
@@ -763,32 +779,45 @@ bool FSequencerTrailHierarchy::CheckForChanges()
 
 void FSequencerTrailHierarchy::UpdateViewAndEvalRange()
 {
-	TSharedPtr<ISequencer> Sequencer = WeakSequencer.Pin();
-	check(Sequencer);
-
-	FFrameRate TickResolution = Sequencer->GetFocusedTickResolution();
-	FFrameRate DisplayRate = Sequencer->GetFocusedDisplayRate();
-	TOptional<TRange<FFrameNumber>> OptionalRange = Sequencer->GetSubSequenceRange();
-	TickEvalRange  = OptionalRange.IsSet()  ? OptionalRange.GetValue() : Sequencer->GetFocusedMovieSceneSequence()->GetMovieScene()->GetPlaybackRange();
-	if (!UMotionTrailToolOptions::GetTrailOptions()->bShowFullTrail)
+	if (TSharedPtr<ISequencer> Sequencer = WeakSequencer.Pin())
 	{
-		FFrameTime SequenceTime = Sequencer->GetLocalTime().Time;
-		const FFrameNumber TicksBefore = FFrameRate::TransformTime(FFrameNumber(UMotionTrailToolOptions::GetTrailOptions()->FramesBefore), DisplayRate, TickResolution).FloorToFrame();
-		const FFrameNumber TicksAfter = FFrameRate::TransformTime(FFrameNumber(UMotionTrailToolOptions::GetTrailOptions()->FramesAfter), DisplayRate, TickResolution).FloorToFrame();
-		TickViewRange = TRange<FFrameNumber>(SequenceTime.GetFrame() - TicksBefore, SequenceTime.GetFrame() + TicksAfter);
-		if (TickViewRange.GetLowerBoundValue() < TickEvalRange.GetLowerBoundValue())
-		{
-			TickViewRange.SetLowerBoundValue(TickEvalRange.GetLowerBoundValue());
-		}
-		if (TickViewRange.GetUpperBoundValue() > TickEvalRange.GetUpperBoundValue())
+
+		FFrameRate TickResolution = Sequencer->GetFocusedTickResolution();
+		FFrameRate DisplayRate = Sequencer->GetFocusedDisplayRate();
+		if (UMotionTrailToolOptions* Settings = UMotionTrailToolOptions::GetTrailOptions())
 		{
-			TickViewRange.SetUpperBoundValue(TickEvalRange.GetUpperBoundValue());
+			//if range is selected then eval and view is selected, otherwise eval is full and view maybe ticksbefore,after
+			if (Settings->MotionTrailRange == EMotionTrailRange::SelectionRange && Sequencer->GetSelectionRange().IsEmpty() == false)
+			{
+				TickEvalRange = Sequencer->GetSelectionRange();
+				TickViewRange = TickEvalRange;
+			}
+			else
+			{
+				TOptional<TRange<FFrameNumber>> OptionalRange = Sequencer->GetSubSequenceRange();
+				TickEvalRange = OptionalRange.IsSet() ? OptionalRange.GetValue() : Sequencer->GetFocusedMovieSceneSequence()->GetMovieScene()->GetPlaybackRange();
+				if (Settings->MotionTrailRange == EMotionTrailRange::SpecifiedRange)
+				{
+					FFrameTime SequenceTime = Sequencer->GetLocalTime().Time;
+					const FFrameNumber TicksBefore = FFrameRate::TransformTime(FFrameNumber(UMotionTrailToolOptions::GetTrailOptions()->FramesBefore), DisplayRate, TickResolution).FloorToFrame();
+					const FFrameNumber TicksAfter = FFrameRate::TransformTime(FFrameNumber(UMotionTrailToolOptions::GetTrailOptions()->FramesAfter), DisplayRate, TickResolution).FloorToFrame();
+					TickViewRange = TRange<FFrameNumber>(SequenceTime.GetFrame() - TicksBefore, SequenceTime.GetFrame() + TicksAfter);
+					if (TickViewRange.GetLowerBoundValue() < TickEvalRange.GetLowerBoundValue())
+					{
+						TickViewRange.SetLowerBoundValue(TickEvalRange.GetLowerBoundValue());
+					}
+					if (TickViewRange.GetUpperBoundValue() > TickEvalRange.GetUpperBoundValue())
+					{
+						TickViewRange.SetUpperBoundValue(TickEvalRange.GetUpperBoundValue());
+					}
+				}
+				else //full range
+				{
+					TickViewRange = TickEvalRange;
+				}
+			}
 		}
 	}
-	else
-	{
-		TickViewRange = Sequencer->GetFocusedMovieSceneSequence()->GetMovieScene()->GetPlaybackRange();
-	}
 }
 
 FGuid FSequencerTrailHierarchy::AddComponentToHierarchy(const FGuid& InBindingGuid, USceneComponent* CompToAdd, UMovieScene3DTransformTrack* TransformTrack)
@@ -1047,25 +1076,27 @@ void FSequencerTrailHierarchy::PinTrail(FGuid InGuid)
 {
 	if (const TUniquePtr<FTrail>* Trail = GetAllTrails().Find(InGuid))
 	{
-		UMotionTrailToolOptions* Settings = UMotionTrailToolOptions::GetTrailOptions();
-		if (Settings->GetNumPinned() >= Settings->MaxNumberPinned)
+		if (UMotionTrailToolOptions* Settings = UMotionTrailToolOptions::GetTrailOptions())
 		{
-			UE_LOG(LogTemp, Warning, TEXT("MotionTrails: Cannot Pin trail %s Max number reached. Please delete pinned trail if you want to add this one."), *(*Trail)->GetName().ToString());
-			return;
-		}
+			if (Settings->GetNumPinned() >= Settings->MaxNumberPinned)
+			{
+				UE_LOG(LogTemp, Warning, TEXT("MotionTrails: Cannot Pin trail %s Max number reached. Please delete pinned trail if you want to add this one."), *(*Trail)->GetName().ToString());
+				return;
+			}
 
-		VisibilityManager.SetTrailAlwaysVisible(InGuid, true);
-		if ((*Trail)->GetDrawInfo())
-		{
-			(*Trail)->GetDrawInfo()->SetColor(UMotionTrailToolOptions::GetTrailOptions()->DefaultColor);
-			(*Trail)->GetDrawInfo()->SetStyle(UMotionTrailToolOptions::GetTrailOptions()->TrailStyle);
-
-			UMotionTrailToolOptions::FPinnedTrail PinnedTrail;
-			PinnedTrail.TrailGuid = InGuid;
-			PinnedTrail.TrailName = (*Trail)->GetName();
-			PinnedTrail.TrailColor = (*Trail)->GetDrawInfo()->GetColor();
-			PinnedTrail.bHasOffset = (*Trail)->HasOffsetTransform();
-			Settings->AddPinned(PinnedTrail);
+			VisibilityManager.SetTrailAlwaysVisible(InGuid, true);
+			if ((*Trail)->GetDrawInfo())
+			{
+				(*Trail)->GetDrawInfo()->SetColor(Settings->DefaultColor);
+				(*Trail)->GetDrawInfo()->SetStyle(Settings->TrailStyle);
+
+				UMotionTrailToolOptions::FPinnedTrail PinnedTrail;
+				PinnedTrail.TrailGuid = InGuid;
+				PinnedTrail.TrailName = (*Trail)->GetName();
+				PinnedTrail.TrailColor = (*Trail)->GetDrawInfo()->GetColor();
+				PinnedTrail.bHasOffset = (*Trail)->HasOffsetTransform();
+				Settings->AddPinned(PinnedTrail);
+			}
 		}
 	}
 }
@@ -1171,8 +1202,8 @@ FGuid FSequencerTrailHierarchy::PinComponent(USceneComponent* InSceneComponent,
 	{
 		if (CurTrail->GetDrawInfo())
 		{
-			CurTrail->GetDrawInfo()->SetColor(UMotionTrailToolOptions::GetTrailOptions()->DefaultColor);
-			CurTrail->GetDrawInfo()->SetStyle(UMotionTrailToolOptions::GetTrailOptions()->TrailStyle);
+			CurTrail->GetDrawInfo()->SetColor(Settings->DefaultColor);
+			CurTrail->GetDrawInfo()->SetStyle(Settings->TrailStyle);
 
 			UMotionTrailToolOptions::FPinnedTrail Trail;
 			Trail.TrailGuid = BoneGuid;
@@ -1213,16 +1244,19 @@ void FSequencerTrailHierarchy::OnPutPinnedInSpace(FGuid InGuid, AActor* InActor,
 {
 	if (const TUniquePtr<FTrail>* Trail = GetAllTrails().Find(InGuid))
 	{
-		if (InActor)
+		if ((*Trail).IsValid())
 		{
-			(*Trail)->SetSpace(InActor, InComponentName);
-			(*Trail)->ForceEvaluateNextTick();
+			if (InActor)
+			{
+				(*Trail)->SetSpace(InActor, InComponentName);
+				(*Trail)->ForceEvaluateNextTick();
 
-		}
-		else
-		{
-			(*Trail)->ClearSpace();
-			(*Trail)->ForceEvaluateNextTick();
+			}
+			else
+			{
+				(*Trail)->ClearSpace();
+				(*Trail)->ForceEvaluateNextTick();
+			}
 		}
 	}
 }
@@ -1231,25 +1265,47 @@ void FSequencerTrailHierarchy::OnSetLinearColor(FGuid InGuid, FLinearColor Color
 {
 	if (const TUniquePtr<FTrail>* Trail = GetAllTrails().Find(InGuid))
 	{
-		if (FTrajectoryDrawInfo* CurDrawInfo = (*Trail)->GetDrawInfo())
+		if ((*Trail).IsValid())
 		{
-			CurDrawInfo->SetColor(Color);
+			if (FTrajectoryDrawInfo* CurDrawInfo = (*Trail)->GetDrawInfo())
+			{
+				CurDrawInfo->SetColor(Color);
+			}
 		}
 	}
 }
+
 void FSequencerTrailHierarchy::OnSetHasOffset(FGuid InGuid, bool bOffset)
 {
 	if (const TUniquePtr<FTrail>* Trail = GetAllTrails().Find(InGuid))
 	{
-		if (bOffset == false)
-		{
-			(*Trail)->ClearOffsetTransform();
-		}
-		else
+		if ((*Trail).IsValid())
 		{
-			(*Trail)->SetOffsetMode();
+			if (bOffset != (*Trail)->GetOffsetMode())
+			{
+				if (bOffset == false)
+				{
+					GLevelEditorModeTools().SetWidgetMode(SavedWidgetMode);
+					GLevelEditorModeTools().SetCoordSystem(SavedCoordSystem);
+					(*Trail)->ClearOffsetTransform();
+				}
+				else
+				{
+					SavedWidgetMode = GLevelEditorModeTools().GetWidgetMode();
+					SavedCoordSystem = GLevelEditorModeTools().GetCoordSystem();
+					GLevelEditorModeTools().SetWidgetMode(UE::Widget::WM_Translate);
+					GLevelEditorModeTools().SetCoordSystem(COORD_Local);
+					(*Trail)->SetOffsetMode();
+				}
+			}
+			else
+			{
+				if (bOffset == false)
+				{
+					(*Trail)->ClearOffsetTransform();
+				}
+			}
 		}
-
 	}
 }
 
diff --git a/Engine/Plugins/Editor/SequencerAnimTools/Source/SequencerAnimTools/Private/Sequencer/SequencerTrailHierarchy.h b/Engine/Plugins/Editor/SequencerAnimTools/Source/SequencerAnimTools/Private/Sequencer/SequencerTrailHierarchy.h
--- a/Engine/Plugins/Editor/SequencerAnimTools/Source/SequencerAnimTools/Private/Sequencer/SequencerTrailHierarchy.h
+++ b/Engine/Plugins/Editor/SequencerAnimTools/Source/SequencerAnimTools/Private/Sequencer/SequencerTrailHierarchy.h
@@ -5,6 +5,7 @@
 #include "TrailHierarchy.h"
 #include "ISequencer.h"
 #include "Tools/EvaluateSequencerTools.h"
+#include "UnrealWidgetFwd.h"
 
 namespace UE::SequencerAnimTools { class FTrail; }
 namespace UE::SequencerAnimTools { class FTrailEvaluateTimes; }
@@ -21,6 +22,12 @@ class UMovieSceneControlRigParameterTrack;
 class UMovieSceneControlRigParameterSection;
 class UControlRig;
 struct FRigHierarchyContainer;
+
+
+namespace UE::Widget
+{
+	enum EWidgetMode;
+};
 namespace UE
 {
 namespace SequencerAnimTools
@@ -123,6 +130,9 @@ public:
 private:
 	//cached guid to check for updates
 	FGuid LastValidMovieSceneGuid;
+	// for offset mode
+	UE::Widget::EWidgetMode SavedWidgetMode;
+	ECoordSystem SavedCoordSystem;
 protected: //callbacks
 	void ClearSelection(); 
 	void RegisterMotionTrailOptionDelegates();
@@ -135,6 +145,7 @@ protected: //callbacks
 	void OnPutPinnedInSpace(FGuid InGuid, AActor* InActor, FName InComponentName);
 	void OnSetLinearColor(FGuid InGuid, FLinearColor Color);
 	void OnSetHasOffset(FGuid InGuid, bool bOffset);
+
 };
 
 
diff --git a/Engine/Plugins/Editor/SequencerAnimTools/Source/SequencerAnimTools/Private/TrailHierarchy.cpp b/Engine/Plugins/Editor/SequencerAnimTools/Source/SequencerAnimTools/Private/TrailHierarchy.cpp
--- a/Engine/Plugins/Editor/SequencerAnimTools/Source/SequencerAnimTools/Private/TrailHierarchy.cpp
+++ b/Engine/Plugins/Editor/SequencerAnimTools/Source/SequencerAnimTools/Private/TrailHierarchy.cpp
@@ -149,13 +149,25 @@ void FTrailHierarchyRenderer::DrawHUD(const FSceneView* View, FCanvas* Canvas)
 				GuidTrailPair.Value->GetTickPointsForDisplay(Transform, *CurrentFramesInfo, false, Marks, MarkNormals);
 				if (Marks.Num() > 0)
 				{
+					FColorState ColorState;
+					ColorState.Setup(OwningHierarchy);
 					const FLinearColor Color = UMotionTrailToolOptions::GetTrailOptions()->MarkColor;
+					const FLinearColor CurrentFrameMarkColor = UMotionTrailToolOptions::GetTrailOptions()->CurrentFrameMarkColor;
 					for (int32 Idx = 0; Idx < Marks.Num(); Idx++)
 					{
+						const FFrameNumber Frame = CurrentFramesInfo->CurrentFrameTimes.CalculateFrame(Idx);
+
 						const FVector2D StartPoint = Marks[Idx] - MarkNormals[Idx] * UMotionTrailToolOptions::GetTrailOptions()->MarkSize;
 						const FVector2D EndPoint = Marks[Idx] + MarkNormals[Idx] * UMotionTrailToolOptions::GetTrailOptions()->MarkSize;
 						FCanvasLineItem LineItem = FCanvasLineItem(StartPoint, EndPoint);
-						LineItem.SetColor(Color);
+						if (ColorState.SequencerTime != Frame)
+						{
+							LineItem.SetColor(Color);
+						}
+						else
+						{
+							LineItem.SetColor(CurrentFrameMarkColor);
+						}
 						Canvas->DrawItem(LineItem);
 					}
 				}
@@ -363,17 +375,24 @@ bool FTrailHierarchy::IsAnythingSelected() const
 	return false;
 }
 
-bool FTrailHierarchy::IsAnythingSelected(FVector& OutVectorPosition) const
+bool FTrailHierarchy::IsAnythingSelected(FVector& OutVectorPosition, FQuat& OutRotation) const
 {
 	OutVectorPosition = FVector::ZeroVector;
+	OutRotation = FQuat::Identity;
 	int NumSelected = 0;
+	TOptional<FQuat> NewRotation;
+	FQuat LocalRotation;
 	for (const TPair<FGuid, TUniquePtr<FTrail>>& GuidTrailPair : GetAllTrails())
 	{
 		FVector Location;
-		bool bHandled = GuidTrailPair.Value->IsAnythingSelected(Location);
+		bool bHandled = GuidTrailPair.Value->IsAnythingSelected(Location,LocalRotation);
 		if (bHandled)
 		{
-			OutVectorPosition += Location;
+			OutVectorPosition += Location; 
+			if (NewRotation.IsSet() == false)
+			{
+				NewRotation = LocalRotation;
+			}
 			++NumSelected;
 		}
 	}
@@ -381,8 +400,11 @@ bool FTrailHierarchy::IsAnythingSelected(FVector& OutVectorPosition) const
 	if (NumSelected > 0)
 	{
 		OutVectorPosition /= (double)NumSelected;
+		OutRotation = NewRotation.GetValue();
+
 	}
 	return (NumSelected > 0);
+
 }
 
 bool FTrailHierarchy::IsAnythingSelected(TArray<FVector>& OutVectorPositions, bool bAllPositions) const
@@ -393,7 +415,8 @@ bool FTrailHierarchy::IsAnythingSelected(TArray<FVector>& OutVectorPositions, bo
 		if (bAllPositions == false)
 		{
 			FVector Location;
-			bool bHandled = GuidTrailPair.Value->IsAnythingSelected(Location);
+			FQuat Rotation;
+			bool bHandled = GuidTrailPair.Value->IsAnythingSelected(Location,Rotation);
 			if (bHandled)
 			{
 				OutVectorPositions.Add(Location);
diff --git a/Engine/Plugins/Editor/SequencerAnimTools/Source/SequencerAnimTools/Private/TrajectoryDrawInfo.cpp b/Engine/Plugins/Editor/SequencerAnimTools/Source/SequencerAnimTools/Private/TrajectoryDrawInfo.cpp
--- a/Engine/Plugins/Editor/SequencerAnimTools/Source/SequencerAnimTools/Private/TrajectoryDrawInfo.cpp
+++ b/Engine/Plugins/Editor/SequencerAnimTools/Source/SequencerAnimTools/Private/TrajectoryDrawInfo.cpp
@@ -105,21 +105,19 @@ bool FFrameCalculator::CalculateIndices()
 		CurrentRange = 0;
 		return true;
 	}
-	if (CurrentBucket >= 0)
+	if (CurrentBucket == -1 && MustHaveIndices.Num() > 0)
 	{
-		for (int32 Index = 0; Index < IndicesToCalculatePerBucket; ++Index)
-		{
-			const int32 CalculatedIndex = CurrentRangeStartIndex + CurrentBucket + (Index * NumOfBuckets);
-			if (CalculatedIndex <= CurrentRangeEndIndex) //can go over
-			{
-				IndicesToCalculate.Add(CalculatedIndex);
-			}
-		}
-		IndicesToCalculate.Add(MustHaveIndices[MustHaveIndices.Num() -1]);
+		IndicesToCalculate.Append(MustHaveIndices);
+		++CurrentBucket;
 	}
-	else
+
+	for (int32 Index = 0; Index < IndicesToCalculatePerBucket; ++Index)
 	{
-		IndicesToCalculate = MustHaveIndices;
+		const int32 CalculatedIndex = CurrentRangeStartIndex + CurrentBucket + (Index * NumOfBuckets);
+		if (CalculatedIndex <= CurrentRangeEndIndex) //can go over
+		{
+			IndicesToCalculate.Add(CalculatedIndex);
+		}
 	}
 	++CurrentBucket;
 	if (CurrentBucket > NumOfBuckets)
@@ -139,11 +137,10 @@ bool FFrameCalculator::CalculateIndices()
 *
 **********************************************************************************/
 
-void FCurrentFramesInfo::SetViewRange(const TRange<FFrameNumber>& InViewRange)
+void FCurrentFramesInfo::SetViewRange(const TRange<FFrameNumber>& InViewRange, bool bInViewRangeIsEvalRange)
 {
 	ViewRange = InViewRange;
-	bViewRangeIsEvalRange = (CurrentFrameTimes.StartFrame == ViewRange.GetLowerBoundValue() &&
-		CurrentFrameTimes.EndFrame == ViewRange.GetUpperBoundValue());
+	bViewRangeIsEvalRange = bInViewRangeIsEvalRange;
 }
 void FCurrentFramesInfo::SetUpFrameTimes(const TRange<FFrameNumber>& InEvalFrameRange, const FFrameNumber& InFrameStep)
 {
diff --git a/Engine/Plugins/Editor/SequencerAnimTools/Source/SequencerAnimTools/Public/Trail.h b/Engine/Plugins/Editor/SequencerAnimTools/Source/SequencerAnimTools/Public/Trail.h
--- a/Engine/Plugins/Editor/SequencerAnimTools/Source/SequencerAnimTools/Public/Trail.h
+++ b/Engine/Plugins/Editor/SequencerAnimTools/Source/SequencerAnimTools/Public/Trail.h
@@ -132,12 +132,13 @@ public:
 	virtual void DrawHUD(const FSceneView* View, FCanvas* Canvas) {};
 	virtual bool HandleClick(const FGuid& Guid, FEditorViewportClient* InViewportClient, HHitProxy* HitProxy, FInputClick Click) { return false; }
 	virtual bool IsAnythingSelected() const { return false; }
-	virtual bool IsAnythingSelected(FVector& OutVectorPosition)const { return false;}
+	virtual bool IsAnythingSelected(FVector& OutVectorPosition, FQuat& OutRotation)const { return false;}
 	virtual bool IsAnythingSelected(TArray<FVector>& OutVectorPositions)const { return false; }
 	virtual bool IsTrailSelected() const { return false; }
 	virtual bool BoxSelect(FBox& InBox, bool InSelect = true) {return false;}
 	virtual bool FrustumSelect(const FConvexVolume& InFrustum, FEditorViewportClient* InViewportClient, bool InSelect = true)  { return false; }
 	virtual void SetOffsetMode() {};
+	virtual bool GetOffsetMode() const {return false;}
 	virtual bool StartTracking() { return false; }
 	virtual bool ApplyDelta(const FVector& Pos, const FRotator& Rot, const FVector& WidgetLocation, bool bApplyToOffset) { return false; }
 	virtual bool EndTracking() { return false; }
diff --git a/Engine/Plugins/Editor/SequencerAnimTools/Source/SequencerAnimTools/Public/TrailHierarchy.h b/Engine/Plugins/Editor/SequencerAnimTools/Source/SequencerAnimTools/Public/TrailHierarchy.h
--- a/Engine/Plugins/Editor/SequencerAnimTools/Source/SequencerAnimTools/Public/TrailHierarchy.h
+++ b/Engine/Plugins/Editor/SequencerAnimTools/Source/SequencerAnimTools/Public/TrailHierarchy.h
@@ -115,8 +115,8 @@ public:
 	virtual bool BoxSelect(FBox& InBox, bool InSelect = true);
 	virtual bool FrustumSelect(const FConvexVolume& InFrustum, FEditorViewportClient* InViewportClient, bool InSelect = true) ;
 
-	//Get calculated center of all selections into one position, return false if not selected
-	virtual bool IsAnythingSelected(FVector& OutVectorPosition)const;
+	//Get calculated center of all selections into one position and rotation, return false if not selected
+	virtual bool IsAnythingSelected(FVector& OutVectorPosition, FQuat& OutRotation)const;
 	//Get each selected item's position, return false if not selected. if bAllPositions is true it will get all selected positions from each trail, otherwise will be the average
 	virtual bool IsAnythingSelected(TArray<FVector>& OutVectorPositions, bool bAllPositions = false) const;
 	virtual bool IsAnythingSelected() const;
diff --git a/Engine/Plugins/Editor/SequencerAnimTools/Source/SequencerAnimTools/Public/TrajectoryDrawInfo.h b/Engine/Plugins/Editor/SequencerAnimTools/Source/SequencerAnimTools/Public/TrajectoryDrawInfo.h
--- a/Engine/Plugins/Editor/SequencerAnimTools/Source/SequencerAnimTools/Public/TrajectoryDrawInfo.h
+++ b/Engine/Plugins/Editor/SequencerAnimTools/Source/SequencerAnimTools/Public/TrajectoryDrawInfo.h
@@ -74,7 +74,7 @@ private:
 // Holds set of currently calculating frames
 struct FCurrentFramesInfo
 {
-	void SetViewRange(const TRange<FFrameNumber>& InViewRange);
+	void SetViewRange(const TRange<FFrameNumber>& InViewRange, bool bInViewRangeIsEvalRange);
 	void SetUpFrameTimes(const TRange<FFrameNumber>& InEvalFrameRange, const FFrameNumber& InFrameStep);
 	void AddMustHaveTimes(const TSet<FFrameNumber>& InMustHaveTimes, const FFrameNumber& InFrameNumber);
 	const TArray<int32>& IndicesToCalculate() const;
diff --git a/Engine/Source/Editor/Sequencer/Public/Tools/MotionTrailOptions.h b/Engine/Source/Editor/Sequencer/Public/Tools/MotionTrailOptions.h
--- a/Engine/Source/Editor/Sequencer/Public/Tools/MotionTrailOptions.h
+++ b/Engine/Source/Editor/Sequencer/Public/Tools/MotionTrailOptions.h
@@ -24,6 +24,15 @@ enum class EMotionTrailTrailStyle : uint8
 	HeatMap = 3
 };
 
+UENUM(BlueprintType)
+enum class EMotionTrailRange : uint8
+{
+	PlaybackRange = 0,
+	SelectionRange = 1,
+	SpecifiedRange = 2,
+	//to do TimeBookmark = 3
+};
+
 
 // TODO: option to make tick size proportional to distance from camera to get a sense of perspective and scale
 UCLASS(MinimalAPI, config = EditorPerProjectUserSettings)
@@ -41,7 +50,10 @@ public:
 		, TimePostColor(.85, 0.25, 0.1)
 		, DashPreColor(0.2, 0.9, 0.3)
 		, DashPostColor(0.7, 0.2, 0.7)
-		, bShowFullTrail(true)
+PRAGMA_DISABLE_DEPRECATION_WARNINGS
+		, bShowFullTrail(false) //deprecrated 5.8
+PRAGMA_ENABLE_DEPRECATION_WARNINGS
+		, MotionTrailRange(EMotionTrailRange::PlaybackRange)
 		, TrailThickness(0.0f)
 		, FramesBefore(10)
 		, FramesAfter(10)
@@ -53,6 +65,7 @@ public:
 		, KeySize(3.0f)
 		, bShowMarks(false)
 		, MarkColor(0.25,1.0,0.15)
+		, CurrentFrameMarkColor(0.4, 0.9, 0.1)
 		, MarkSize(5.0)
 		, MaxNumberPinned(10)
 		, bLockMarksToFrames_DEPRECATED(true)
@@ -91,20 +104,23 @@ public:
 	UPROPERTY(EditAnywhere, Config, Category = AdvancedDisplay)
 	FLinearColor DashPostColor;
 
-	/** Whether or not to show the full motion trail */
-	UPROPERTY(EditAnywhere, Category = Trail);
+	UE_DEPRECATED(5.8, "Use Motion Trail Range")
 	bool bShowFullTrail;
 
+	/** The time range to draw the trail */
+	UPROPERTY(EditAnywhere, Config, Category = Trail);
+	EMotionTrailRange MotionTrailRange;
+
 	/* The thickness of the motion trail */
 	UPROPERTY(EditAnywhere, Config, Category = Trail, Meta = (ClampMin = "0.0"))
 	double TrailThickness;
 
 	/** The number of frames to draw before the start of the trail. Requires not showing the full trail */
-	UPROPERTY(EditAnywhere, Config, Category = Trail, Meta = (EditCondition = "!bShowFullTrail", ClampMin = "0"))
+	UPROPERTY(EditAnywhere, Config, Category = Trail, Meta = (EditCondition = "MotionTrailRange == EMotionTrailRange::SpecifiedRange", ClampMin = "0"))
 	int32 FramesBefore;
 
 	/** The number of frames to draw after the end of the trail. Requires not showing the full trail */
-	UPROPERTY(EditAnywhere, Config, Category = Trail, Meta = (EditCondition = "!bShowFullTrail", ClampMin = "0"))
+	UPROPERTY(EditAnywhere, Config, Category = Trail, Meta = (EditCondition = "MotionTrailRange == EMotionTrailRange::SpecifiedRange", ClampMin = "0"))
 	int32 FramesAfter;
 
 	/** No longer exposed and clamped to 1 The number of evaluations per frame */
@@ -138,6 +154,10 @@ public:
 	UPROPERTY(EditAnywhere, Config, Category = Marks)
 	FLinearColor MarkColor;
 
+	/** The color of the Current Frame mark color */
+	UPROPERTY(EditAnywhere, Config, Category = Marks)
+	FLinearColor CurrentFrameMarkColor;
+
 	/** The size of the marks */
 	UPROPERTY(EditAnywhere, Config, Category = Marks, Meta = (ClampMin = "0.0"))
 	double MarkSize;
