Pivot Tool improvements
- better overall interactivity
- bAutokeyOnRelease taken into account (as well as sequencer autokey)
- controls are sorted by dependency for updating in order
- constraints taken into account
- bound components updated for better visual feedback

#rnx

#Synced-CL 49749792

diff --git a/Engine/Plugins/Animation/ControlRig/Source/ControlRigEditor/Public/Sequencer/EditModeAnimationUtil.h b/Engine/Plugins/Animation/ControlRig/Source/ControlRigEditor/Public/Sequencer/EditModeAnimationUtil.h
--- a/Engine/Plugins/Animation/ControlRig/Source/ControlRigEditor/Public/Sequencer/EditModeAnimationUtil.h
+++ b/Engine/Plugins/Animation/ControlRig/Source/ControlRigEditor/Public/Sequencer/EditModeAnimationUtil.h
@@ -87,7 +87,7 @@ struct FControlKeyframeData
  * or any other widget that would need to deffer keyframing. 
  */
 	
-struct FControlRigKeyframer
+struct CONTROLRIGEDITOR_API FControlRigKeyframer
 {
 	~FControlRigKeyframer();
 	
diff --git a/Engine/Plugins/Editor/SequencerAnimTools/Source/SequencerAnimTools/Private/SequencerAnimEditPivotTool.cpp b/Engine/Plugins/Editor/SequencerAnimTools/Source/SequencerAnimTools/Private/SequencerAnimEditPivotTool.cpp
--- a/Engine/Plugins/Editor/SequencerAnimTools/Source/SequencerAnimTools/Private/SequencerAnimEditPivotTool.cpp
+++ b/Engine/Plugins/Editor/SequencerAnimTools/Source/SequencerAnimTools/Private/SequencerAnimEditPivotTool.cpp
@@ -14,22 +14,23 @@
 #include "InteractiveToolManager.h"
 // for raycast into World
 
+#include "ConstraintsManager.h"
 #include "EditorModeManager.h"
 #include "Misc/QualifiedFrameTime.h"
 #include "Modules/ModuleManager.h"
 #include "MovieSceneTimeUnit.h"
 #include "ControlRig.h"
+#include "ControlRigComponent.h"
 #include "ControlRigSequencerEditorLibrary.h"
 #include "EditorInteractiveGizmoManager.h"
 #include "LevelSequence.h"
 #include "LevelSequenceEditorBlueprintLibrary.h"
 #include "ILevelSequenceEditorToolkit.h"
-#include "InteractiveToolManager.h"
 #include "Tools/EdModeInteractiveToolsContext.h"
 #include "SceneView.h"
 #include "UnrealEdGlobals.h"
 #include "Editor/UnrealEdEngine.h"
-#include "Modules/ModuleManager.h"
+#include "Rigs/FKControlRig.h"
 #include "Selection.h"
 #include "Widgets/Text/STextBlock.h"
 #include "Widgets/Input/SButton.h"
@@ -37,7 +38,8 @@
 #include "Widgets/Layout/SUniformGridPanel.h"
 #include "Styling/StyleColors.h"
 #include "SLevelViewport.h"
-#include "EditorGizmos/EditorTransformGizmoUtil.h"
+#include "Constraints/ControlRigTransformableHandle.h"
+#include "Constraints/MovieSceneConstraintChannelHelper.h"
 #include "EditorGizmos/TransformGizmo.h"
 #include "EditMode/ControlRigEditMode.h"
 
@@ -84,6 +86,19 @@ namespace UE::AIE::PivotTool
 
 		return Names;
 	}
+	
+	USceneComponent* GetBoundComponent(const UControlRig* InControlRig)
+	{
+		const TSharedPtr<IControlRigObjectBinding>& ObjectBinding = InControlRig ? InControlRig->GetObjectBinding() : nullptr;
+		return ObjectBinding.IsValid() ? Cast<USceneComponent>(ObjectBinding->GetBoundObject()) : nullptr;
+	}
+	
+	FTransform GetBoundComponentTransform(const UControlRig* InControlRig)
+	{
+		USceneComponent* HostingComponent = GetBoundComponent(InControlRig);
+		return HostingComponent ? HostingComponent->GetComponentTransform() : FTransform::Identity;
+	}
+	
 	//includes global + hosting transform
 	FTransform GetControlWorldSpaceTransform(UControlRig* ControlRig, const FName& Control)
 	{
@@ -92,15 +107,7 @@ namespace UE::AIE::PivotTool
 		{
 			if (FRigTransformElement* ControlElement = ControlRig->FindControl(Control))
 			{
-				TSharedPtr<IControlRigObjectBinding> ObjectBinding = ControlRig->GetObjectBinding();
-				FTransform ParentTransform = FTransform::Identity;
-				if (ObjectBinding.IsValid())
-				{
-					if (USceneComponent* HostingComponent = Cast<USceneComponent>(ObjectBinding->GetBoundObject()))
-					{
-						ParentTransform = HostingComponent->GetComponentTransform();
-					}
-				}
+				const FTransform ParentTransform = GetBoundComponentTransform(ControlRig);
 				//Use GetTransform() not GetControlGlobalTransform
 				WorldTransform = ControlRig->GetHierarchy()->GetTransform(ControlElement, ERigTransformType::CurrentGlobal);
 				WorldTransform = WorldTransform * ParentTransform;
@@ -131,7 +138,6 @@ static void GetControlRigsAndSequencer(TArray<TWeakObjectPtr<UControlRig>>& Cont
 			}
 			for (UControlRig* ControlRig : TempControlRigs)
 			{
-				UControlRig* Base = ControlRig;
 				ControlRigs.Add(ControlRig);
 			}
 		}
@@ -177,19 +183,6 @@ UInteractiveTool* USequencerPivotToolBuilder::BuildTool(const FToolBuilderState&
 /*
 * ControlRig/Actor Mappings
 */
-FTransform FControlRigMappings::GetParentTransform() const
-{
-	if (ControlRig.IsValid())
-	{
-		TSharedPtr<IControlRigObjectBinding> ObjectBinding = ControlRig->GetObjectBinding();
-		if (ObjectBinding.IsValid())
-		{
-			USceneComponent* HostingComponent =  Cast<USceneComponent>(ObjectBinding->GetBoundObject());
-			return HostingComponent ? HostingComponent->GetComponentTransform() : FTransform::Identity;
-		}
-	}
-	return FTransform::Identity;
-}
 
 TOptional<FTransform> FControlRigMappings::GetWorldTransform(const FName& Name) const
 {
@@ -202,7 +195,7 @@ TOptional<FTransform> FControlRigMappings::GetWorldTransform(const FName& Name)
 			const FTransform* LocalTransform = PivotTransforms.Find(Name);
 			if (LocalTransform)
 			{
-				FTransform ToWorldTransform = GetParentTransform();
+				const FTransform ToWorldTransform = UE::AIE::PivotTool::GetBoundComponentTransform(ControlRig.Get());
 				//Use GetTransform() not GetControlGlobalTransform
 				FTransform CurrentTransform = ControlRig->GetHierarchy()->GetTransform(ControlElement, ERigTransformType::CurrentGlobal);
 				WorldTransform = (*LocalTransform) * CurrentTransform * ToWorldTransform;
@@ -218,7 +211,7 @@ void FControlRigMappings::SetFromWorldTransform(const FName& Name, const FTransf
 	{
 		if (FRigTransformElement* ControlElement = ControlRig->FindControl(Name))
 		{
-			FTransform ToWorldTransform = GetParentTransform();
+			FTransform ToWorldTransform = UE::AIE::PivotTool::GetBoundComponentTransform(ControlRig.Get());
 			FTransform GlobalTransform = ControlRig->GetHierarchy()->GetTransform(ControlElement, ERigTransformType::CurrentGlobal) * ToWorldTransform;
 
 			FTransform LocalTransform = WorldTransform.GetRelativeTransform(GlobalTransform);
@@ -401,6 +394,8 @@ void USequencerPivotTool::Setup()
 	SetPivotMode(false);
 	CreateAndShowPivotOverlay();
 	
+	Keyframer.Initialize();
+	ConstraintsCache.RegisterNotifications();
 }
 
 void USequencerPivotTool::SetPivotMode(bool bVal)
@@ -411,6 +406,7 @@ void USequencerPivotTool::SetPivotMode(bool bVal)
 		FEditorModeTools& ModeManager = GLevelEditorModeTools(); //mz need to fix for persona/ animation edit mode
 		ModeManager.SetCoordSystem(COORD_Local);
 	}
+	Keyframer.Enable(bInPivotMode);
 }
 
 void USequencerPivotTool::OnTick(float DeltaTime) 
@@ -589,6 +585,10 @@ void USequencerPivotTool::Shutdown(EToolShutdownType ShutdownType)
 	GizmoManager->DestroyAllGizmosByOwner(this);
 	RemoveDelegates();
 	RemoveAndDestroyPivotOverlay();
+	
+	InteractionDependencies.Reset();
+	ConstraintsCache.UnregisterNotifications();
+	ConstraintsCache.Reset();
 }
 
 void USequencerPivotTool::DeactivateMe()
@@ -688,87 +688,33 @@ void USequencerPivotTool::RemoveDelegates()
 	OnEditorSelectionChangedHandle.Reset();
 }
 
-
 void USequencerPivotTool::GizmoTransformStarted(UTransformProxy* Proxy)
 {
+	using namespace UE::AIE::PivotTool;
+	
+	Keyframer.Enable(bInPivotMode);
+	InteractionDependencies.Reset();
+	TransformContext = FControlRigInteractionTransformContext();
+	
 	ControlRigDrags.SetNum(0);
-	TransactionIndex = INDEX_NONE;
+	
 	ISequencer* Sequencer = SequencerPtr.Pin().Get();
 	if (Sequencer && LevelSequence)
 	{
-		//TransactionIndex = GEditor->BeginTransaction(nullptr, LOCTEXT("TransformPivot", "Transform Pivot"), nullptr);
-		const FFrameRate TickResolution = Sequencer->GetFocusedTickResolution();
-		const FFrameTime FrameTime = Sequencer->GetLocalTime().ConvertTo(TickResolution);
-		const FFrameNumber FrameNumber = FrameTime.RoundToFrame();
-
-		for (TWeakObjectPtr<UControlRig>& ControlRig : ControlRigs)
+		const bool bHandleTransaction = GEditor && !GIsTransacting && GEditor->CanTransact() && !GEditor->IsTransactionActive();
+		if (bHandleTransaction)
 		{
-			if (ControlRig.IsValid() == false)
-			{
-				continue;
-			}
-			TArray<FName> SelectedControls = UE::AIE::PivotTool::CurrentControlSelection(ControlRig.Get());
-			if (SelectedControls.Num() > 0)
-			{
-				ControlRig->Modify();
-
-				for (const FName& Name : SelectedControls)
-				{
-					if (FRigControlElement* ControlElement = ControlRig->FindControl(Name))
-					{
-						//if we have a parent that's selected and we are a float/bool/int/enum skip
-						if (FRigControlElement* ParentControlElement = Cast<FRigControlElement>(ControlRig->GetHierarchy()->GetFirstParent(ControlElement)))
-						{
-							if (ControlRig->IsControlSelected(ParentControlElement->GetFName()) && (
-								ControlElement->Settings.ControlType == ERigControlType::Bool || 
-								ControlElement->Settings.ControlType == ERigControlType::Float || 
-								ControlElement->Settings.ControlType == ERigControlType::ScaleFloat || 
-								ControlElement->Settings.ControlType == ERigControlType::Integer)
-								)
-							{
-								continue;
-							}
-						}
-						UControlRig* RealControlRig = Cast<UControlRig>(ControlRig.Get());
-						FTransform Transform = UE::AIE::PivotTool::GetControlWorldSpaceTransform(RealControlRig, Name);
-
-						FControlRigSelectionDuringDrag ControlDrag;
-						ControlDrag.LevelSequence = LevelSequence;
-						ControlDrag.ControlName = Name;
-						ControlDrag.ControlRig = ControlRig.Get();
-						ControlDrag.CurrentFrame = FrameNumber;
-						ControlDrag.CurrentTransform = Transform;
-						ControlRigDrags.Add(ControlDrag);
-					}
-				}
-			}
+			Transaction = MakeUnique<FScopedTransaction>(LOCTEXT("USequencerPivotToolTransaction", "Transform"));
+		}
+		
+		for (TWeakObjectPtr<UControlRig>& WeakControlRig : ControlRigs)
+		{
+			PrepareRigForManipulation(WeakControlRig.Get());
 		}
 
 		for (TWeakObjectPtr<AActor>& Actor : Actors)
 		{
-			if (Actor.IsValid() == false)
-			{
-				continue;
-			}
-			if (USceneComponent* RootComponent = Actor->GetRootComponent())
-			{
-				if (RootComponent->HasAnyFlags(RF_DefaultSubObject))
-				{
-					// Default subobjects must be included in any undo/redo operations
-					RootComponent->SetFlags(RF_Transactional);
-				}
-				RootComponent->Modify();
-			}
-
-			FTransform Transform = UControlRigSequencerEditorLibrary::GetActorWorldTransform(LevelSequence, Actor.Get(), FrameNumber,
-				EMovieSceneTimeUnit::TickResolution);
-
-			FActorSelectonDuringDrag ActorDrag;
-			ActorDrag.LevelSequence = LevelSequence;
-			ActorDrag.Actor = Actor.Get();
-			ActorDrag.CurrentFrame = FrameNumber;
-			ActorDrag.CurrentTransform = Transform;
-			ActorDrags.Add(ActorDrag);
+			PrepareActorForManipulation(Actor.Get());
 		}
 	}
 
@@ -778,6 +724,7 @@ void USequencerPivotTool::GizmoTransformStarted(UTransformProxy* Proxy)
 	GizmoTransform = StartDragTransform;
 
 	InteractionScopes.Reset();
+
 	if(bInPivotMode)
 	{
 		TMap<UControlRig*, int32> RigToScopeIndex;
@@ -837,6 +784,89 @@ static void SetLocation(AActor* Actor, const FVector& NewLocation,const TOptiona
 	}
 }
 
+void USequencerPivotTool::SetControlTransform(const FControlRigSelectionDuringDrag& InControlDrag)
+{
+	UControlRig* ControlRig = InControlDrag.ControlRig;
+	if (!ControlRig || InControlDrag.ControlName == NAME_None)
+	{
+		return;
+	}
+	
+	const bool bFixEulerFlips = (!ControlRig->IsAdditive() || ControlRig->IsA<UFKControlRig>()) && TransformContext.bRotation;
+	UE::ControlRigEditMode::FInteractionDependencyCache* Dependencies = InteractionDependencies.Find(ControlRig);
+	auto EvaluateRigIfNeeded = [Dependencies, &InControlDrag]()
+	{
+		// skip compensation and evaluate the rig to force notifications: auto-key and constraints updates (among others) are based on
+		// UControlRig::OnControlModified being broadcast but this only happens on evaluation for additive rigs.
+		// constraint compensation is disabled while manipulating in that case to avoid re-entrant evaluations
+		TGuardValue<bool> CompensateGuard(FMovieSceneConstraintChannelHelper::bDoNotCompensate, true);
+			
+		if (InControlDrag.ControlRig->IsAdditive())
+		{
+			InControlDrag.ControlRig->Evaluate_AnyThread();
+		}
+		else if (Dependencies)
+		{
+			if (Dependencies->HasDownwardDependencies({InControlDrag.ControlName, ERigElementType::Control}) ||
+				Dependencies->CheckAndUpdateParentsPoseVersion())
+			{
+				InControlDrag.ControlRig->Evaluate_AnyThread();
+			}
+		}
+	};
+	
+	const bool bDeferAutokey = Keyframer.IsEnabled();
+	const uint32 ControlHash = UTransformableControlHandle::ComputeHash(ControlRig, InControlDrag.ControlName);
+	const bool bNeedsConstraintPostProcess = ConstraintsCache.HasAnyActiveConstraint(ControlHash, TargetWorld);
+	static const TArray< TWeakObjectPtr<UTickableConstraint> > EmptyConstraints;
+	const TArray< TWeakObjectPtr<UTickableConstraint> >& Constraints =
+		bNeedsConstraintPostProcess ? ConstraintsCache.Get(ControlHash, TargetWorld) : EmptyConstraints;
+
+	FRigControlModifiedContext Context;
+	Context.bConstraintUpdate = true;
+	Context.SetKey = bDeferAutokey ? EControlRigSetKey::Never : EControlRigSetKey::DoNotCare;
+
+	USceneComponent* Component = UE::AIE::PivotTool::GetBoundComponent(ControlRig);
+	const FTransform& ComponentTransform = Component ? Component->GetComponentTransform() : FTransform::Identity;
+	
+	const FTransform Global = InControlDrag.CurrentTransform.GetRelativeTransform(ComponentTransform);
+	
+	constexpr bool bNotify = true, bSetupUndo = true, bPrintPythonCommands = false;
+	ControlRig->SetControlGlobalTransform(
+		InControlDrag.ControlName, Global, bNotify, Context, bSetupUndo, bPrintPythonCommands, bFixEulerFlips);
+	EvaluateRigIfNeeded();
+				
+	UE::AnimationEditMode::FControlKeyframeData KeyframeData;
+	KeyframeData.LocalTransform = ControlRig->GetControlLocalTransform(InControlDrag.ControlName);
+	
+	if (bNeedsConstraintPostProcess)
+	{
+		// switch to constraint space
+		const TOptional<FTransform> RelativeTransform =
+			UE::TransformConstraintUtil::GetConstraintsRelativeTransform(Constraints, KeyframeData.LocalTransform, InControlDrag.CurrentTransform);
+		if (RelativeTransform)
+		{
+			KeyframeData.LocalTransform = *RelativeTransform;
+			KeyframeData.bConstraintSpace = true;
+		}
+
+		FRigControlModifiedContext LocalContext = Context;
+		LocalContext.bConstraintUpdate = false;
+
+		ControlRig->SetControlLocalTransform(InControlDrag.ControlName, KeyframeData.LocalTransform, bNotify, LocalContext, bSetupUndo, bFixEulerFlips);
+		EvaluateRigIfNeeded();
+	}
+
+	if (bNeedsConstraintPostProcess || ConstraintsCache.HasAnyActiveConstraint(Component, TargetWorld))
+	{
+		TGuardValue<bool> CompensateGuard(FMovieSceneConstraintChannelHelper::bDoNotCompensate, true);
+		const FConstraintsManagerController& Controller = FConstraintsManagerController::Get(TargetWorld);
+		Controller.EvaluateAllConstraints();
+	}
+	
+	Keyframer.Store(ControlHash, MoveTemp(KeyframeData));
+}
+
 void USequencerPivotTool::GizmoTransformChanged(UTransformProxy* Proxy, FTransform Transform)
 {
 	if (!bGizmoBeingDragged)
@@ -844,27 +874,26 @@ void USequencerPivotTool::GizmoTransformChanged(UTransformProxy* Proxy, FTransfo
 		return;
 	}
 	bManipulatorMadeChange = true;
-
+	
 	FModifierKeysState KeyState = FSlateApplication::Get().GetModifierKeys();
 	bool bAlsoCalcRotation = KeyState.IsShiftDown() == false;
 	GizmoTransform = Transform;
 	FTransform Diff = Transform.GetRelativeTransform(StartDragTransform);
-	const bool bRotationChanged = (Diff.GetRotation().IsIdentity(1e-4f) == false);
-	const bool bTranslationChanged = (Diff.GetTranslation().IsNearlyZero() == false);
-	if (bRotationChanged || bTranslationChanged)
+	TransformContext.bRotation = (Diff.GetRotation().IsIdentity(1e-4f) == false);
+	TransformContext.bTranslation = (Diff.GetTranslation().IsNearlyZero() == false);
+	
+	if (TransformContext.bRotation || TransformContext.bTranslation)
 	{
-		const bool bSetKey = false;
 		int32 Index = 0;
 		for (FControlRigSelectionDuringDrag& ControlDrag : ControlRigDrags)
 		{
-			UControlRig* RealControlRig = Cast<UControlRig>(ControlDrag.ControlRig);
-			if (bInPivotMode == false)
+			if (!bInPivotMode)
 			{
 				ControlDrag.CurrentTransform = Transform;
 			}
 			else if (!bShiftPressedWhenStarted || Index != (ControlRigDrags.Num() - 1))
 			{
-				if (bRotationChanged)
+				if (TransformContext.bRotation)
 				{
 					FVector LocDiff = ControlDrag.CurrentTransform.GetLocation() - Transform.GetLocation();
 					if (LocDiff.IsNearlyZero(1e-4f) == false)
@@ -874,6 +903,7 @@ void USequencerPivotTool::GizmoTransformChanged(UTransformProxy* Proxy, FTransfo
 						FVector NewLocation = StartDragTransform.TransformPosition(RotatedDiff);
 						ControlDrag.CurrentTransform.SetLocation(NewLocation);
 					}
+					
 					if (bAlsoCalcRotation)
 					{
 						FQuat StartDragRot = StartDragTransform.GetRotation();
@@ -885,22 +915,19 @@ void USequencerPivotTool::GizmoTransformChanged(UTransformProxy* Proxy, FTransfo
 						OptQuat.EnforceShortestArcWith(ControlDrag.CurrentTransform.GetRotation());
 						ControlDrag.CurrentTransform.SetRotation(OptQuat);
 					}
-					UControlRigSequencerEditorLibrary::SetControlRigWorldTransform(ControlDrag.LevelSequence, RealControlRig, ControlDrag.ControlName,
-						ControlDrag.CurrentFrame, ControlDrag.CurrentTransform, EMovieSceneTimeUnit::TickResolution, bSetKey);
 					
+					SetControlTransform(ControlDrag);
 				}
-				else if (bTranslationChanged)
+				else if (TransformContext.bTranslation)
 				{
 					FVector DiffTranslation = StartDragTransform.GetRotation().RotateVector(Diff.GetTranslation());
 					ControlDrag.CurrentTransform.AddToTranslation(DiffTranslation);
-					UControlRigSequencerEditorLibrary::SetControlRigWorldTransform(ControlDrag.LevelSequence, RealControlRig, ControlDrag.ControlName,
-						ControlDrag.CurrentFrame, ControlDrag.CurrentTransform, EMovieSceneTimeUnit::TickResolution, bSetKey);
+					SetControlTransform(ControlDrag);
 				}
 			}
 			else //last one with shift we keep locked!
 			{
-				UControlRigSequencerEditorLibrary::SetControlRigWorldTransform(ControlDrag.LevelSequence, RealControlRig, ControlDrag.ControlName,
-					ControlDrag.CurrentFrame, ControlDrag.CurrentTransform, EMovieSceneTimeUnit::TickResolution, bSetKey);
+				SetControlTransform(ControlDrag);
 			}
 			
 			++Index;
@@ -915,7 +942,7 @@ void USequencerPivotTool::GizmoTransformChanged(UTransformProxy* Proxy, FTransfo
 			}
 			else if (!bShiftPressedWhenStarted || Index != (ActorDrags.Num() - 1))
 			{
-				if (bRotationChanged)
+				if (TransformContext.bRotation)
 				{
 					FVector LocDiff = ActorDrag.CurrentTransform.GetLocation() - Transform.GetLocation();
 					if (LocDiff.IsNearlyZero(1e-4f) == false)
@@ -950,7 +977,7 @@ void USequencerPivotTool::GizmoTransformChanged(UTransformProxy* Proxy, FTransfo
 						SetLocation(ActorDrag.Actor, ActorDrag.CurrentTransform.GetLocation(), OptQuat);
 					}
 				}
-				else if (bTranslationChanged)
+				else if (TransformContext.bTranslation)
 				{
 					FVector DiffTranslation = StartDragTransform.GetRotation().RotateVector(Diff.GetTranslation());
 					ActorDrag.CurrentTransform.AddToTranslation(DiffTranslation);
@@ -969,24 +996,40 @@ void USequencerPivotTool::GizmoTransformChanged(UTransformProxy* Proxy, FTransfo
 	}
 	
 	StartDragTransform = Transform;
+	
+	UpdateBoundComponents();
 }
 
 void USequencerPivotTool::GizmoTransformEnded(UTransformProxy* Proxy)
 {
-
-	if (bManipulatorMadeChange && TransactionIndex != INDEX_NONE)
+	const bool bFinalizeTransaction = Transaction && bManipulatorMadeChange;
+	
+	bGizmoBeingDragged = false;
+	bManipulatorMadeChange = false;
+	
+	// autokey on release
 	{
-		//GEditor->EndTransaction();
+		for (const TSharedPtr<FControlRigInteractionScope>& InteractionScope: InteractionScopes)
+		{
+			Keyframer.Apply(*InteractionScope, TransformContext);
+		}
+		Keyframer.Finalize(TargetWorld);
+		Keyframer.Reset();
 	}
-	else if (TransactionIndex != INDEX_NONE)
+	
+	if (Transaction)
 	{
-	//	GEditor->CancelTransaction(TransactionIndex);
+		if (!bFinalizeTransaction)
+		{
+			Transaction->Cancel();
+		}
+		Transaction.Reset();
 	}
-	bGizmoBeingDragged = false;
-	bManipulatorMadeChange = false;
+	
 	InteractionScopes.Reset();
 	UpdateGizmoTransform();
 	SavePivotTransforms();
+	InteractionDependencies.Reset();
 }
 
 void USequencerPivotTool::SavePivotTransforms()
@@ -1434,6 +1477,219 @@ void USequencerPivotTool::UpdatePivotOverlayLocation(const FVector2D InLocation,
 	LastPivotOverlayLocation = ScreenPos;
 }
 
+void USequencerPivotTool::PrepareRigForManipulation(UControlRig* InControlRig)
+{
+	using namespace UE::AIE::PivotTool;
+	
+	URigHierarchy* Hierarchy = InControlRig ? InControlRig->GetHierarchy() : nullptr;
+	if (!Hierarchy)
+	{
+		return;
+	}
+		
+	const TArray<FName> SelectedControls = CurrentControlSelection(InControlRig);
+	if (SelectedControls.IsEmpty())
+	{
+		return;
+	}
+
+	// store controls to transform
+	TArray<FRigControlElement*> Controls;
+	Controls.Reserve(SelectedControls.Num());
+	for (const FName& Name : SelectedControls)
+	{
+		if (FRigControlElement* ControlElement = InControlRig->FindControl(Name))
+		{
+			//if we have a parent that's selected and we are a float/bool/int/enum skip
+			if (FRigControlElement* ParentControlElement = Cast<FRigControlElement>(Hierarchy->GetFirstParent(ControlElement)))
+			{
+				if (InControlRig->IsControlSelected(ParentControlElement->GetFName()) && (
+					ControlElement->Settings.ControlType == ERigControlType::Bool || 
+					ControlElement->Settings.ControlType == ERigControlType::Float || 
+					ControlElement->Settings.ControlType == ERigControlType::ScaleFloat || 
+					ControlElement->Settings.ControlType == ERigControlType::Integer)
+					)
+				{
+					continue;
+				}
+			}
+			Controls.Add(ControlElement);
+		}
+	}
+			
+	if (Controls.IsEmpty())
+	{
+		return;
+	}
+			
+	InControlRig->Modify();
+	
+	// keep last selected control if bShiftPressedWhenStarted
+	FRigControlElement* LastControl = nullptr;
+	if (bShiftPressedWhenStarted && !Controls.IsEmpty())
+	{
+		LastControl = Controls.Last();
+		Controls.RemoveAt(Controls.Num() - 1);
+	}
+
+	// sort control from depth to avoid setting a child before it's parent
+	Algo::Sort(Controls, [Hierarchy](const FRigControlElement* Control0, const FRigControlElement* Control1)
+	{
+		constexpr bool bRecursive = true;	
+		const FRigBaseElementParentArray& Parents = Hierarchy->GetParents(Control0, bRecursive);
+		const FRigBaseElementParentArray& OtherParents = Hierarchy->GetParents(Control1, bRecursive);
+
+		int32 Index = Parents.Num(), OtherIndex = OtherParents.Num();
+		auto HasCommonAncestor = [&Index, &Parents, &OtherIndex, &OtherParents, Control1, Control0]()
+		{
+			if (!Parents.IsEmpty() && !OtherParents.IsEmpty())
+			{
+				FRigBaseElement* CommonAncestor = nullptr;
+				while (Index > 0 && OtherIndex > 0 && Parents[Index-1] == OtherParents[OtherIndex-1])
+				{
+					CommonAncestor = Parents[Index-1];
+					--Index;
+					--OtherIndex;
+				}
+				return CommonAncestor != nullptr;
+			}
+			return false;
+		};
+
+		// if common ancestor, compare depths
+		if (HasCommonAncestor() && Index != OtherIndex)
+		{
+			return OtherIndex > Index; 
+		}
+
+		// otherwise, compare indices
+		if (Control1->GetIndex() > Control0->GetIndex())
+		{
+			return true;
+		}
+
+		return false;
+	});
+	
+	// re-add last selected control stored if bShiftPressedWhenStarted
+	if (LastControl)
+	{
+		Controls.Add(LastControl);
+	}
+			
+	UE::ControlRigEditMode::FInteractionDependencyCache& NewDependency = InteractionDependencies.Add(InControlRig);
+	NewDependency.WeakHierarchy = Hierarchy;
+	
+	const FRigDependenciesProviderForControlRig DependencyProvider(InControlRig);
+	for (int32 Index = 0; Index < Controls.Num(); ++Index)
+	{
+		FRigControlElement* Control = Controls[Index];
+		
+		FControlRigSelectionDuringDrag ControlDrag;
+		ControlDrag.LevelSequence = LevelSequence;
+		ControlDrag.ControlName = Control->GetFName();
+		ControlDrag.ControlRig = InControlRig;
+		ControlDrag.CurrentTransform = GetControlWorldSpaceTransform(ControlDrag.ControlRig, ControlDrag.ControlName);
+		ControlRigDrags.Add(ControlDrag);
+				
+		for (int32 NextIndex = Index+1; NextIndex < Controls.Num(); ++NextIndex)
+		{
+			FRigControlElement* OtherControl = Controls[NextIndex];
+			if (Hierarchy->IsParentedTo(Control, OtherControl, DependencyProvider))
+			{
+				NewDependency.Parents.Add(OtherControl->GetKey());
+				NewDependency.Children.Add(Control->GetKey());
+			}
+			else if (Hierarchy->IsParentedTo(OtherControl, Control, DependencyProvider))
+			{
+				NewDependency.Parents.Add(Control->GetKey());
+				NewDependency.Children.Add(OtherControl->GetKey());
+			}
+		}
+
+		// store parents' pose versions
+		const FRigBaseElementParentArray Parents = Hierarchy->GetParents(Control);
+		NewDependency.ParentsPoseVersion.Reserve(Parents.Num());
+		for (const FRigBaseElement* Parent: Parents)
+		{
+			if (const FRigTransformElement* TransformParent = Cast<FRigTransformElement>(Parent))
+			{
+				NewDependency.ParentsPoseVersion.FindOrAdd(Parent->GetIndex()) = Hierarchy->GetPoseVersion(TransformParent);
+			}
+		}
+	}
+}
+
+void USequencerPivotTool::PrepareActorForManipulation(AActor* InActor)
+{
+	if (!InActor)
+	{
+		return;
+	}
+	
+	if (USceneComponent* RootComponent = InActor->GetRootComponent())
+	{
+		if (RootComponent->HasAnyFlags(RF_DefaultSubObject))
+		{
+			// Default subobjects must be included in any undo/redo operations
+			RootComponent->SetFlags(RF_Transactional);
+		}
+		RootComponent->Modify();
+	}
+
+	FActorSelectonDuringDrag ActorDrag;
+	ActorDrag.LevelSequence = LevelSequence;
+	ActorDrag.Actor = InActor;
+	ActorDrag.CurrentTransform = InActor->GetTransform();
+	ActorDrags.Add(ActorDrag);
+}
+
+void USequencerPivotTool::UpdateBoundComponents()
+{
+	if (!bManipulatorMadeChange)
+	{
+		return;
+	}
+	
+	for (const TSharedPtr<FControlRigInteractionScope>& InteractionScope: InteractionScopes)
+	{
+		if (UControlRig* ControlRig = InteractionScope->GetControlRig())
+		{
+			USceneComponent* SceneComponent = UE::AIE::PivotTool::GetBoundComponent(ControlRig);
+			if (UControlRigComponent* ControlRigComponent = Cast<UControlRigComponent>(SceneComponent))
+			{
+				ControlRigComponent->Update();
+			}
+			else if (USkeletalMeshComponent* SkeletalMeshComponent = Cast<USkeletalMeshComponent>(SceneComponent))
+			{
+				if (!ControlRig->ContainsSimulation())
+				{
+					constexpr float DeltaTime = 0.f;
+					constexpr bool bIncludeFromChildActors = true, bDoNotNeedValidRootMotion = false;
+					const AActor* Parent = SkeletalMeshComponent->GetOwner();
+					while (Parent)
+					{
+						Parent->ForEachComponent<USkeletalMeshComponent>(bIncludeFromChildActors, [](USkeletalMeshComponent* SkeletalMeshToUpdate)
+						{
+							if (SkeletalMeshToUpdate)
+							{
+								SkeletalMeshToUpdate->TickAnimation(DeltaTime, bDoNotNeedValidRootMotion);
+								SkeletalMeshToUpdate->RefreshBoneTransforms();
+								SkeletalMeshToUpdate->RefreshFollowerComponents	();
+								SkeletalMeshToUpdate->UpdateComponentToWorld();
+								SkeletalMeshToUpdate->FinalizeBoneTransform();
+								SkeletalMeshToUpdate->MarkRenderTransformDirty();
+								SkeletalMeshToUpdate->MarkRenderDynamicDataDirty();
+							}								
+						});
+						Parent = Parent->GetAttachParentActor();
+					}
+				}
+			}
+		}
+	}
+}
+
 #undef LOCTEXT_NAMESPACE
 
 
diff --git a/Engine/Plugins/Editor/SequencerAnimTools/Source/SequencerAnimTools/Public/SequencerAnimEditPivotTool.h b/Engine/Plugins/Editor/SequencerAnimTools/Source/SequencerAnimTools/Public/SequencerAnimEditPivotTool.h
--- a/Engine/Plugins/Editor/SequencerAnimTools/Source/SequencerAnimTools/Public/SequencerAnimEditPivotTool.h
+++ b/Engine/Plugins/Editor/SequencerAnimTools/Source/SequencerAnimTools/Public/SequencerAnimEditPivotTool.h
@@ -8,7 +8,12 @@
 #include "BaseBehaviors/BehaviorTargetInterfaces.h"
 #include "Framework/Commands/Commands.h"
 #include "BaseSequencerAnimTool.h"
+#include "EditMode/ControlRigEditModeUtil.h"
+#include "ScopedTransaction.h"
+#include "Sequencer/EditModeAnimationUtil.h"
 #include "Styling/AppStyle.h"
+#include "Transform/TransformConstraintUtil.h"
+
 #include "SequencerAnimEditPivotTool.generated.h"
 
 class FControlRigInteractionScope;
@@ -52,10 +57,9 @@ public:
 struct FSelectionDuringDrag
 {
 	ULevelSequence* LevelSequence;
-	FFrameNumber CurrentFrame;
 	FTransform CurrentTransform;
-
 };
+
 struct FControlRigSelectionDuringDrag : public FSelectionDuringDrag
 {
 	UControlRig* ControlRig;
@@ -71,7 +75,6 @@ struct FControlRigMappings
 {
 	TWeakObjectPtr<UControlRig> ControlRig;
 
-	FTransform GetParentTransform() const;
 	TOptional<FTransform> GetWorldTransform(const FName& Name) const;
 	void SetFromWorldTransform(const FName& Name, const FTransform& WorldTransform);
 	TArray<FName> GetAllControls() const;
@@ -193,7 +196,7 @@ protected:
 	FTransform StartDragTransform;
 	bool bGizmoBeingDragged = false;
 	bool bManipulatorMadeChange = false;
-	int32 TransactionIndex = -1;
+	TUniquePtr<FScopedTransaction> Transaction = nullptr;
 	TArray<FControlRigSelectionDuringDrag> ControlRigDrags;
 	TArray<FActorSelectonDuringDrag> ActorDrags;
 	TArray<TSharedPtr<FControlRigInteractionScope>> InteractionScopes;
@@ -247,6 +250,30 @@ private:
 	TSharedPtr<FUICommandList> CommandBindings;
 	void ResetPivot();
 
+	// Stores data related to InControlRig to prepare its manipulation
+	void PrepareRigForManipulation(UControlRig* InControlRig);
+	
+	// Stores data related to InActor to prepare its manipulation
+	void PrepareActorForManipulation(AActor* InActor);
+	
+	// Updates a control using provided FControlRigSelectionDuringDrag data.
+	void SetControlTransform(const FControlRigSelectionDuringDrag& InControlDrag);
+	
+	// Current transform context.
+	FControlRigInteractionTransformContext TransformContext;
+	
+	// Used to store and apply keyframes. (if deferred)
+	UE::AnimationEditMode::FControlRigKeyframer Keyframer;
+	
+	// Per control constraints cache.
+	UE::TransformConstraintUtil::FConstraintsInteractionCache ConstraintsCache;
+	
+	// Per ControlRig dependencies between the selected controls during interaction.
+	TMap<UControlRig*, UE::ControlRigEditMode::FInteractionDependencyCache> InteractionDependencies;
+	
+	// Updates the components bound to the control rigs for a responsive feedback
+	void UpdateBoundComponents();
+
 public:
 	static FSavedMappings SavedPivotLocations;
 	static FLastSelectedObjects LastSelectedObjects;
diff --git a/Engine/Plugins/Editor/SequencerAnimTools/Source/SequencerAnimTools/SequencerAnimTools.Build.cs b/Engine/Plugins/Editor/SequencerAnimTools/Source/SequencerAnimTools/SequencerAnimTools.Build.cs
--- a/Engine/Plugins/Editor/SequencerAnimTools/Source/SequencerAnimTools/SequencerAnimTools.Build.cs
+++ b/Engine/Plugins/Editor/SequencerAnimTools/Source/SequencerAnimTools/SequencerAnimTools.Build.cs
@@ -22,6 +22,7 @@ public class SequencerAnimTools: ModuleRules
 			new string[]
 			{
 				"CoreUObject",
+				"Constraints",
 				"Engine",
 				"InputCore",
 				"UnrealEd",
