Pivot Tool: only modify the hierachy when changing pivot

#rnx

#Synced-CL 50857758

diff --git a/Engine/Plugins/Editor/SequencerAnimTools/Source/SequencerAnimTools/Private/SequencerAnimEditPivotTool.cpp b/Engine/Plugins/Editor/SequencerAnimTools/Source/SequencerAnimTools/Private/SequencerAnimEditPivotTool.cpp
--- a/Engine/Plugins/Editor/SequencerAnimTools/Source/SequencerAnimTools/Private/SequencerAnimEditPivotTool.cpp
+++ b/Engine/Plugins/Editor/SequencerAnimTools/Source/SequencerAnimTools/Private/SequencerAnimEditPivotTool.cpp
@@ -1522,7 +1522,7 @@ void USequencerPivotTool::PrepareRigForManipulation(UControlRig* InControlRig)
 		return;
 	}
 			
-	InControlRig->Modify();
+	Hierarchy->Modify();
 	
 	// keep last selected control if bShiftPressedWhenStarted
 	FRigControlElement* LastControl = nullptr;
