Attempt to fix PVS static analysis warning with explicit lambda capture.

#Synced-CL 51964473

diff --git a/Engine/Plugins/Editor/SequencerAnimTools/Source/SequencerAnimTools/Private/MotionTrailTool.cpp b/Engine/Plugins/Editor/SequencerAnimTools/Source/SequencerAnimTools/Private/MotionTrailTool.cpp
--- a/Engine/Plugins/Editor/SequencerAnimTools/Source/SequencerAnimTools/Private/MotionTrailTool.cpp
+++ b/Engine/Plugins/Editor/SequencerAnimTools/Source/SequencerAnimTools/Private/MotionTrailTool.cpp
@@ -385,7 +385,7 @@ void UMotionTrailTool::UpdateGizmoLocation()
 		int NumSelected = 0;
 		TOptional<FQuat> NewRotation;
 		FQuat LocalRotation;
-		ForEachActiveHierarchy([&](UE::SequencerAnimTools::FTrailHierarchy& Hierarchy)
+		ForEachActiveHierarchy([&NumSelected, &LocalRotation, &NewRotation, &NewGizmoLocation](UE::SequencerAnimTools::FTrailHierarchy& Hierarchy)
 		{
 			FVector Location;
 			if (Hierarchy.IsAnythingSelected(Location, LocalRotation))
