Silence false positive PVS.

#Synced-CL 51964473

diff --git a/Engine/Plugins/Editor/SequencerAnimTools/Source/SequencerAnimTools/Private/MotionTrailTool.cpp b/Engine/Plugins/Editor/SequencerAnimTools/Source/SequencerAnimTools/Private/MotionTrailTool.cpp
--- a/Engine/Plugins/Editor/SequencerAnimTools/Source/SequencerAnimTools/Private/MotionTrailTool.cpp
+++ b/Engine/Plugins/Editor/SequencerAnimTools/Source/SequencerAnimTools/Private/MotionTrailTool.cpp
@@ -398,7 +398,7 @@ void UMotionTrailTool::UpdateGizmoLocation()
 				++NumSelected;
 			}
 		});
-		if (NumSelected > 0)
+		if (NumSelected > 0)//-V547
 		{
 			NewGizmoLocation /= (double)NumSelected;
 			LocalRotation = NewRotation.GetValue();
